L10N: P4: Fixed changelist numbers showing up comma delimited.

diff --git a/Engine/Plugins/Developer/PerforceSourceControl/Source/PerforceSourceControl/Private/PerforceSourceControlOperations.cpp b/Engine/Plugins/Developer/PerforceSourceControl/Source/PerforceSourceControl/Private/PerforceSourceControlOperations.cpp
--- a/Engine/Plugins/Developer/PerforceSourceControl/Source/PerforceSourceControl/Private/PerforceSourceControlOperations.cpp
+++ b/Engine/Plugins/Developer/PerforceSourceControl/Source/PerforceSourceControl/Private/PerforceSourceControlOperations.cpp
@@ -269,7 +269,7 @@ static FText ParseSubmitResults(const FP4RecordSet& InRecords)
 		const FString SubmittedChange = ClientRecord(TEXT("submittedChange"));
 		if(SubmittedChange.Len() > 0)
 		{
-			return FText::Format(LOCTEXT("SubmitMessage", "Submitted changelist {0}"), FText::AsNumber(FCString::Atoi(*SubmittedChange)));
+			return FText::Format(LOCTEXT("SubmitMessage", "Submitted changelist {0}"), FText::FromString(SubmittedChange));
 		}
 	}
 
