UE4 Refactoring. Changed OVERRIDE and FINAL macros to keywords override and final.

diff --git a/Engine/Plugins/Developer/BlankPlugin/Source/BlankPlugin/Private/BlankPlugin.cpp b/Engine/Plugins/Developer/BlankPlugin/Source/BlankPlugin/Private/BlankPlugin.cpp
--- a/Engine/Plugins/Developer/BlankPlugin/Source/BlankPlugin/Private/BlankPlugin.cpp
+++ b/Engine/Plugins/Developer/BlankPlugin/Source/BlankPlugin/Private/BlankPlugin.cpp
@@ -6,8 +6,8 @@
 class FBlankPlugin : public IBlankPlugin
 {
 	/** IModuleInterface implementation */
-	virtual void StartupModule() OVERRIDE;
-	virtual void ShutdownModule() OVERRIDE;
+	virtual void StartupModule() override;
+	virtual void ShutdownModule() override;
 };
 
 IMPLEMENT_MODULE( FBlankPlugin, BlankPlugin )
diff --git a/Engine/Plugins/Developer/PerforceSourceControl/Source/PerforceSourceControl/Private/PerforceConnection.cpp b/Engine/Plugins/Developer/PerforceSourceControl/Source/PerforceSourceControl/Private/PerforceConnection.cpp
--- a/Engine/Plugins/Developer/PerforceSourceControl/Source/PerforceSourceControl/Private/PerforceConnection.cpp
+++ b/Engine/Plugins/Developer/PerforceSourceControl/Source/PerforceSourceControl/Private/PerforceConnection.cpp
@@ -132,7 +132,7 @@ public:
 	{
 	}
 
-	virtual int IsAlive() OVERRIDE
+	virtual int IsAlive() override
 	{
 		if(IsCancelled.IsBound() && IsCancelled.Execute())
 		{
diff --git a/Engine/Plugins/Developer/PerforceSourceControl/Source/PerforceSourceControl/Private/PerforceSourceControlCommand.h b/Engine/Plugins/Developer/PerforceSourceControl/Source/PerforceSourceControl/Private/PerforceSourceControlCommand.h
--- a/Engine/Plugins/Developer/PerforceSourceControl/Source/PerforceSourceControl/Private/PerforceSourceControlCommand.h
+++ b/Engine/Plugins/Developer/PerforceSourceControl/Source/PerforceSourceControl/Private/PerforceSourceControlCommand.h
@@ -27,13 +27,13 @@ public:
 	 * abandoned before completion. NOTE: This requires the object to delete
 	 * itself using whatever heap it was allocated in.
 	 */
-	virtual void Abandon() OVERRIDE;
+	virtual void Abandon() override;
 
 	/**
 	 * This method is also used to tell the object to cleanup but not before
 	 * the object has finished it's work.
 	 */ 
-	virtual void DoThreadedWork() OVERRIDE;
+	virtual void DoThreadedWork() override;
 
 	/** Attempt to cancel the operation */
 	void Cancel();
diff --git a/Engine/Plugins/Developer/PerforceSourceControl/Source/PerforceSourceControl/Private/PerforceSourceControlLabel.h b/Engine/Plugins/Developer/PerforceSourceControl/Source/PerforceSourceControl/Private/PerforceSourceControlLabel.h
--- a/Engine/Plugins/Developer/PerforceSourceControl/Source/PerforceSourceControl/Private/PerforceSourceControlLabel.h
+++ b/Engine/Plugins/Developer/PerforceSourceControl/Source/PerforceSourceControl/Private/PerforceSourceControlLabel.h
@@ -17,9 +17,9 @@ public:
 	}
 
 	/** ISourceControlLabel implementation */
-	virtual const FString& GetName() const OVERRIDE;
-	virtual bool GetFileRevisions( const TArray<FString>& InFiles, TArray< TSharedRef<ISourceControlRevision, ESPMode::ThreadSafe> >& OutRevisions ) const OVERRIDE;
-	virtual bool Sync( const FString& InFilename ) const OVERRIDE;
+	virtual const FString& GetName() const override;
+	virtual bool GetFileRevisions( const TArray<FString>& InFiles, TArray< TSharedRef<ISourceControlRevision, ESPMode::ThreadSafe> >& OutRevisions ) const override;
+	virtual bool Sync( const FString& InFilename ) const override;
 
 private:
 
diff --git a/Engine/Plugins/Developer/PerforceSourceControl/Source/PerforceSourceControl/Private/PerforceSourceControlModule.h b/Engine/Plugins/Developer/PerforceSourceControl/Source/PerforceSourceControl/Private/PerforceSourceControlModule.h
--- a/Engine/Plugins/Developer/PerforceSourceControl/Source/PerforceSourceControl/Private/PerforceSourceControlModule.h
+++ b/Engine/Plugins/Developer/PerforceSourceControl/Source/PerforceSourceControl/Private/PerforceSourceControlModule.h
@@ -9,8 +9,8 @@ class FPerforceSourceControlModule : public IModuleInterface
 {
 public:
 	/** IModuleInterface implementation */
-	virtual void StartupModule() OVERRIDE;
-	virtual void ShutdownModule() OVERRIDE;
+	virtual void StartupModule() override;
+	virtual void ShutdownModule() override;
 
 	/** Access the Perforce source control settings */
 	FPerforceSourceControlSettings& AccessSettings();
diff --git a/Engine/Plugins/Developer/PerforceSourceControl/Source/PerforceSourceControl/Private/PerforceSourceControlOperations.h b/Engine/Plugins/Developer/PerforceSourceControl/Source/PerforceSourceControl/Private/PerforceSourceControlOperations.h
--- a/Engine/Plugins/Developer/PerforceSourceControl/Source/PerforceSourceControl/Private/PerforceSourceControlOperations.h
+++ b/Engine/Plugins/Developer/PerforceSourceControl/Source/PerforceSourceControl/Private/PerforceSourceControlOperations.h
@@ -9,18 +9,18 @@ class FPerforceConnectWorker : public IPerforceSourceControlWorker
 {
 public:
 	// IPerforceSourceControlWorker interface
-	virtual FName GetName() const OVERRIDE;
-	virtual bool Execute(class FPerforceSourceControlCommand& InCommand) OVERRIDE;
-	virtual bool UpdateStates() const OVERRIDE;
+	virtual FName GetName() const override;
+	virtual bool Execute(class FPerforceSourceControlCommand& InCommand) override;
+	virtual bool UpdateStates() const override;
 };
 
 class FPerforceCheckOutWorker : public IPerforceSourceControlWorker
 {
 public:
 	// IPerforceSourceControlWorker interface
-	virtual FName GetName() const OVERRIDE;
-	virtual bool Execute(class FPerforceSourceControlCommand& InCommand) OVERRIDE;
-	virtual bool UpdateStates() const OVERRIDE;
+	virtual FName GetName() const override;
+	virtual bool Execute(class FPerforceSourceControlCommand& InCommand) override;
+	virtual bool UpdateStates() const override;
 
 public:
 	/** Map of filenames to perforce state */
@@ -31,9 +31,9 @@ class FPerforceCheckInWorker : public IPerforceSourceControlWorker
 {
 public:
 	// IPerforceSourceControlWorker interface
-	virtual FName GetName() const OVERRIDE;
-	virtual bool Execute(class FPerforceSourceControlCommand& InCommand) OVERRIDE;
-	virtual bool UpdateStates() const OVERRIDE;
+	virtual FName GetName() const override;
+	virtual bool Execute(class FPerforceSourceControlCommand& InCommand) override;
+	virtual bool UpdateStates() const override;
 
 public:
 	/** Map of filenames to perforce state */
@@ -47,9 +47,9 @@ class FPerforceMarkForAddWorker : public IPerforceSourceControlWorker
 {
 public:
 	// IPerforceSourceControlWorker interface
-	virtual FName GetName() const OVERRIDE;
-	virtual bool Execute(class FPerforceSourceControlCommand& InCommand) OVERRIDE;
-	virtual bool UpdateStates() const OVERRIDE;
+	virtual FName GetName() const override;
+	virtual bool Execute(class FPerforceSourceControlCommand& InCommand) override;
+	virtual bool UpdateStates() const override;
 
 public:
 	/** Map of filenames to perforce state */
@@ -60,9 +60,9 @@ class FPerforceDeleteWorker : public IPerforceSourceControlWorker
 {
 public:
 	// IPerforceSourceControlWorker interface
-	virtual FName GetName() const OVERRIDE;
-	virtual bool Execute(class FPerforceSourceControlCommand& InCommand) OVERRIDE;
-	virtual bool UpdateStates() const OVERRIDE;
+	virtual FName GetName() const override;
+	virtual bool Execute(class FPerforceSourceControlCommand& InCommand) override;
+	virtual bool UpdateStates() const override;
 
 public:
 	/** Map of filenames to perforce state */
@@ -73,9 +73,9 @@ class FPerforceRevertWorker : public IPerforceSourceControlWorker
 {
 public:
 	// IPerforceSourceControlWorker interface
-	virtual FName GetName() const OVERRIDE;
-	virtual bool Execute(class FPerforceSourceControlCommand& InCommand) OVERRIDE;
-	virtual bool UpdateStates() const OVERRIDE;
+	virtual FName GetName() const override;
+	virtual bool Execute(class FPerforceSourceControlCommand& InCommand) override;
+	virtual bool UpdateStates() const override;
 
 public:
 	/** Map of filenames to perforce state */
@@ -86,9 +86,9 @@ class FPerforceSyncWorker : public IPerforceSourceControlWorker
 {
 public:
 	// IPerforceSourceControlWorker interface
-	virtual FName GetName() const OVERRIDE;
-	virtual bool Execute(class FPerforceSourceControlCommand& InCommand) OVERRIDE;
-	virtual bool UpdateStates() const OVERRIDE;
+	virtual FName GetName() const override;
+	virtual bool Execute(class FPerforceSourceControlCommand& InCommand) override;
+	virtual bool UpdateStates() const override;
 
 public:
 	/** Map of filenames to perforce state */
@@ -99,9 +99,9 @@ class FPerforceUpdateStatusWorker : public IPerforceSourceControlWorker
 {
 public:
 	// IPerforceSourceControlWorker interface
-	virtual FName GetName() const OVERRIDE;
-	virtual bool Execute(class FPerforceSourceControlCommand& InCommand) OVERRIDE;
-	virtual bool UpdateStates() const OVERRIDE;
+	virtual FName GetName() const override;
+	virtual bool Execute(class FPerforceSourceControlCommand& InCommand) override;
+	virtual bool UpdateStates() const override;
 
 public:
 	/** Temporary states for results */
@@ -122,18 +122,18 @@ class FPerforceGetWorkspacesWorker : public IPerforceSourceControlWorker
 {
 public:
 	// IPerforceSourceControlWorker interface
-	virtual FName GetName() const OVERRIDE;
-	virtual bool Execute(class FPerforceSourceControlCommand& InCommand) OVERRIDE;
-	virtual bool UpdateStates() const OVERRIDE;
+	virtual FName GetName() const override;
+	virtual bool Execute(class FPerforceSourceControlCommand& InCommand) override;
+	virtual bool UpdateStates() const override;
 };
 
 class FPerforceCopyWorker : public IPerforceSourceControlWorker
 {
 public:
 	// IPerforceSourceControlWorker interface
-	virtual FName GetName() const OVERRIDE;
-	virtual bool Execute(class FPerforceSourceControlCommand& InCommand) OVERRIDE;
-	virtual bool UpdateStates() const OVERRIDE;
+	virtual FName GetName() const override;
+	virtual bool Execute(class FPerforceSourceControlCommand& InCommand) override;
+	virtual bool UpdateStates() const override;
 
 public:
 	/** Map of filenames to perforce state */
diff --git a/Engine/Plugins/Developer/PerforceSourceControl/Source/PerforceSourceControl/Private/PerforceSourceControlProvider.h b/Engine/Plugins/Developer/PerforceSourceControl/Source/PerforceSourceControl/Private/PerforceSourceControlProvider.h
--- a/Engine/Plugins/Developer/PerforceSourceControl/Source/PerforceSourceControl/Private/PerforceSourceControlProvider.h
+++ b/Engine/Plugins/Developer/PerforceSourceControl/Source/PerforceSourceControl/Private/PerforceSourceControlProvider.h
@@ -23,23 +23,23 @@ public:
 	}
 
 	/* ISourceControlProvider implementation */
-	virtual void Init(bool bForceConnection = true) OVERRIDE;
-	virtual void Close() OVERRIDE;
-	virtual FText GetStatusText() const OVERRIDE;
-	virtual bool IsEnabled() const OVERRIDE;
-	virtual bool IsAvailable() const OVERRIDE;
-	virtual const FName& GetName(void) const OVERRIDE;
-	virtual ECommandResult::Type GetState( const TArray<FString>& InFiles, TArray< TSharedRef<ISourceControlState, ESPMode::ThreadSafe> >& OutState, EStateCacheUsage::Type InStateCacheUsage ) OVERRIDE;
-	virtual void RegisterSourceControlStateChanged( const FSourceControlStateChanged::FDelegate& SourceControlStateChanged ) OVERRIDE;
-	virtual void UnregisterSourceControlStateChanged( const FSourceControlStateChanged::FDelegate& SourceControlStateChanged ) OVERRIDE;
-	virtual ECommandResult::Type Execute( const TSharedRef<ISourceControlOperation, ESPMode::ThreadSafe>& InOperation, const TArray<FString>& InFiles, EConcurrency::Type InConcurrency = EConcurrency::Synchronous, const FSourceControlOperationComplete& InOperationCompleteDelegate = FSourceControlOperationComplete() ) OVERRIDE;
-	virtual bool CanCancelOperation( const TSharedRef<ISourceControlOperation, ESPMode::ThreadSafe>& InOperation ) const OVERRIDE;
-	virtual void CancelOperation( const TSharedRef<ISourceControlOperation, ESPMode::ThreadSafe>& InOperation ) OVERRIDE;
-	virtual bool UsesLocalReadOnlyState() const OVERRIDE;
-	virtual void Tick() OVERRIDE;
-	virtual TArray< TSharedRef<class ISourceControlLabel> > GetLabels( const FString& InMatchingSpec ) const OVERRIDE;
+	virtual void Init(bool bForceConnection = true) override;
+	virtual void Close() override;
+	virtual FText GetStatusText() const override;
+	virtual bool IsEnabled() const override;
+	virtual bool IsAvailable() const override;
+	virtual const FName& GetName(void) const override;
+	virtual ECommandResult::Type GetState( const TArray<FString>& InFiles, TArray< TSharedRef<ISourceControlState, ESPMode::ThreadSafe> >& OutState, EStateCacheUsage::Type InStateCacheUsage ) override;
+	virtual void RegisterSourceControlStateChanged( const FSourceControlStateChanged::FDelegate& SourceControlStateChanged ) override;
+	virtual void UnregisterSourceControlStateChanged( const FSourceControlStateChanged::FDelegate& SourceControlStateChanged ) override;
+	virtual ECommandResult::Type Execute( const TSharedRef<ISourceControlOperation, ESPMode::ThreadSafe>& InOperation, const TArray<FString>& InFiles, EConcurrency::Type InConcurrency = EConcurrency::Synchronous, const FSourceControlOperationComplete& InOperationCompleteDelegate = FSourceControlOperationComplete() ) override;
+	virtual bool CanCancelOperation( const TSharedRef<ISourceControlOperation, ESPMode::ThreadSafe>& InOperation ) const override;
+	virtual void CancelOperation( const TSharedRef<ISourceControlOperation, ESPMode::ThreadSafe>& InOperation ) override;
+	virtual bool UsesLocalReadOnlyState() const override;
+	virtual void Tick() override;
+	virtual TArray< TSharedRef<class ISourceControlLabel> > GetLabels( const FString& InMatchingSpec ) const override;
 #if SOURCE_CONTROL_WITH_SLATE
-	virtual TSharedRef<class SWidget> MakeSettingsWidget() const OVERRIDE;
+	virtual TSharedRef<class SWidget> MakeSettingsWidget() const override;
 #endif
 
 	/**
diff --git a/Engine/Plugins/Developer/PerforceSourceControl/Source/PerforceSourceControl/Private/PerforceSourceControlRevision.h b/Engine/Plugins/Developer/PerforceSourceControl/Source/PerforceSourceControl/Private/PerforceSourceControlRevision.h
--- a/Engine/Plugins/Developer/PerforceSourceControl/Source/PerforceSourceControl/Private/PerforceSourceControlRevision.h
+++ b/Engine/Plugins/Developer/PerforceSourceControl/Source/PerforceSourceControl/Private/PerforceSourceControlRevision.h
@@ -16,19 +16,19 @@ public:
 	}
 
 	/** ISourceControlRevision interface */
-	virtual bool Get( FString& InOutFilename ) const OVERRIDE;
-	virtual bool GetAnnotated( TArray<FAnnotationLine>& OutLines ) const OVERRIDE;
-	virtual bool GetAnnotated( FString& InOutFilename ) const OVERRIDE;
-	virtual const FString& GetFilename() const OVERRIDE;
-	virtual int32 GetRevisionNumber() const OVERRIDE;
-	virtual const FString& GetDescription() const OVERRIDE;
-	virtual const FString& GetUserName() const OVERRIDE;
-	virtual const FString& GetClientSpec() const OVERRIDE;
-	virtual const FString& GetAction() const OVERRIDE;
-	virtual TSharedPtr<ISourceControlRevision, ESPMode::ThreadSafe> GetBranchSource() const OVERRIDE;
-	virtual const FDateTime& GetDate() const OVERRIDE;
-	virtual int32 GetCheckInIdentifier() const OVERRIDE;
-	virtual int32 GetFileSize() const OVERRIDE;
+	virtual bool Get( FString& InOutFilename ) const override;
+	virtual bool GetAnnotated( TArray<FAnnotationLine>& OutLines ) const override;
+	virtual bool GetAnnotated( FString& InOutFilename ) const override;
+	virtual const FString& GetFilename() const override;
+	virtual int32 GetRevisionNumber() const override;
+	virtual const FString& GetDescription() const override;
+	virtual const FString& GetUserName() const override;
+	virtual const FString& GetClientSpec() const override;
+	virtual const FString& GetAction() const override;
+	virtual TSharedPtr<ISourceControlRevision, ESPMode::ThreadSafe> GetBranchSource() const override;
+	virtual const FDateTime& GetDate() const override;
+	virtual int32 GetCheckInIdentifier() const override;
+	virtual int32 GetFileSize() const override;
 
 public:
 	/** The local filename the this revision refers to */
diff --git a/Engine/Plugins/Developer/PerforceSourceControl/Source/PerforceSourceControl/Private/PerforceSourceControlState.h b/Engine/Plugins/Developer/PerforceSourceControl/Source/PerforceSourceControl/Private/PerforceSourceControlState.h
--- a/Engine/Plugins/Developer/PerforceSourceControl/Source/PerforceSourceControl/Private/PerforceSourceControlState.h
+++ b/Engine/Plugins/Developer/PerforceSourceControl/Source/PerforceSourceControl/Private/PerforceSourceControlState.h
@@ -56,28 +56,28 @@ public:
 	}
 
 	/** ISourceControlState interface */
-	virtual int32 GetHistorySize() const OVERRIDE;
-	virtual TSharedPtr<class ISourceControlRevision, ESPMode::ThreadSafe> GetHistoryItem( int32 HistoryIndex ) const OVERRIDE;
-	virtual TSharedPtr<class ISourceControlRevision, ESPMode::ThreadSafe> FindHistoryRevision( int32 RevisionNumber ) const OVERRIDE;
-	virtual FName GetIconName() const OVERRIDE;
-	virtual FName GetSmallIconName() const OVERRIDE;
-	virtual FText GetDisplayName() const OVERRIDE;
-	virtual FText GetDisplayTooltip() const OVERRIDE;
-	virtual const FString& GetFilename() const OVERRIDE;
-	virtual const FDateTime& GetTimeStamp() const OVERRIDE;
-	virtual bool CanCheckIn() const OVERRIDE;
-	virtual bool CanCheckout() const OVERRIDE;
-	virtual bool IsCheckedOut() const OVERRIDE;
-	virtual bool IsCheckedOutOther(FString* Who = NULL) const OVERRIDE;
-	virtual bool IsCurrent() const OVERRIDE;
-	virtual bool IsSourceControlled() const OVERRIDE;
-	virtual bool IsAdded() const OVERRIDE;
-	virtual bool IsDeleted() const OVERRIDE;
-	virtual bool IsIgnored() const OVERRIDE;
-	virtual bool CanEdit() const OVERRIDE;
-	virtual bool IsUnknown() const OVERRIDE;
-	virtual bool IsModified() const OVERRIDE;
-	virtual bool CanAdd() const OVERRIDE;
+	virtual int32 GetHistorySize() const override;
+	virtual TSharedPtr<class ISourceControlRevision, ESPMode::ThreadSafe> GetHistoryItem( int32 HistoryIndex ) const override;
+	virtual TSharedPtr<class ISourceControlRevision, ESPMode::ThreadSafe> FindHistoryRevision( int32 RevisionNumber ) const override;
+	virtual FName GetIconName() const override;
+	virtual FName GetSmallIconName() const override;
+	virtual FText GetDisplayName() const override;
+	virtual FText GetDisplayTooltip() const override;
+	virtual const FString& GetFilename() const override;
+	virtual const FDateTime& GetTimeStamp() const override;
+	virtual bool CanCheckIn() const override;
+	virtual bool CanCheckout() const override;
+	virtual bool IsCheckedOut() const override;
+	virtual bool IsCheckedOutOther(FString* Who = NULL) const override;
+	virtual bool IsCurrent() const override;
+	virtual bool IsSourceControlled() const override;
+	virtual bool IsAdded() const override;
+	virtual bool IsDeleted() const override;
+	virtual bool IsIgnored() const override;
+	virtual bool CanEdit() const override;
+	virtual bool IsUnknown() const override;
+	virtual bool IsModified() const override;
+	virtual bool CanAdd() const override;
 
 	/** Get the state of a file */
 	EPerforceState::Type GetState() const
diff --git a/Engine/Plugins/Developer/PerforceSourceControl/Source/PerforceSourceControl/Private/SPerforceSourceControlSettings.h b/Engine/Plugins/Developer/PerforceSourceControl/Source/PerforceSourceControl/Private/SPerforceSourceControlSettings.h
--- a/Engine/Plugins/Developer/PerforceSourceControl/Source/PerforceSourceControl/Private/SPerforceSourceControlSettings.h
+++ b/Engine/Plugins/Developer/PerforceSourceControl/Source/PerforceSourceControl/Private/SPerforceSourceControlSettings.h
@@ -9,7 +9,7 @@ class FGetWorkspaces : public ISourceControlOperation
 {
 public:
 	// ISourceControlOperation interface
-	virtual FName GetName() const OVERRIDE 
+	virtual FName GetName() const override 
 	{ 
 		return "GetWorkspaces"; 
 	}
diff --git a/Engine/Plugins/Developer/VisualStudioSourceCodeAccess/Source/VisualStudioSourceCodeAccess/Private/VisualStudioSourceCodeAccessor.h b/Engine/Plugins/Developer/VisualStudioSourceCodeAccess/Source/VisualStudioSourceCodeAccess/Private/VisualStudioSourceCodeAccessor.h
--- a/Engine/Plugins/Developer/VisualStudioSourceCodeAccess/Source/VisualStudioSourceCodeAccess/Private/VisualStudioSourceCodeAccessor.h
+++ b/Engine/Plugins/Developer/VisualStudioSourceCodeAccess/Source/VisualStudioSourceCodeAccess/Private/VisualStudioSourceCodeAccessor.h
@@ -22,15 +22,15 @@ public:
 	void Shutdown();
 
 	/** ISourceCodeAccessor implementation */
-	virtual bool CanAccessSourceCode() const OVERRIDE;
-	virtual FName GetFName() const OVERRIDE;
-	virtual FText GetNameText() const OVERRIDE;
-	virtual FText GetDescriptionText() const OVERRIDE;
-	virtual bool OpenSolution() OVERRIDE;
-	virtual bool OpenFileAtLine(const FString& FullPath, int32 LineNumber, int32 ColumnNumber = 0) OVERRIDE;
-	virtual bool OpenSourceFiles(const TArray<FString>& AbsoluteSourcePaths) OVERRIDE;
-	virtual bool SaveAllOpenDocuments() const OVERRIDE;
-	virtual void Tick(const float DeltaTime) OVERRIDE;
+	virtual bool CanAccessSourceCode() const override;
+	virtual FName GetFName() const override;
+	virtual FText GetNameText() const override;
+	virtual FText GetDescriptionText() const override;
+	virtual bool OpenSolution() override;
+	virtual bool OpenFileAtLine(const FString& FullPath, int32 LineNumber, int32 ColumnNumber = 0) override;
+	virtual bool OpenSourceFiles(const TArray<FString>& AbsoluteSourcePaths) override;
+	virtual bool SaveAllOpenDocuments() const override;
+	virtual void Tick(const float DeltaTime) override;
 
 private:
 
diff --git a/Engine/Plugins/Developer/XCodeSourceCodeAccess/Source/XCodeSourceCodeAccess/Private/XCodeSourceCodeAccessor.h b/Engine/Plugins/Developer/XCodeSourceCodeAccess/Source/XCodeSourceCodeAccess/Private/XCodeSourceCodeAccessor.h
--- a/Engine/Plugins/Developer/XCodeSourceCodeAccess/Source/XCodeSourceCodeAccess/Private/XCodeSourceCodeAccessor.h
+++ b/Engine/Plugins/Developer/XCodeSourceCodeAccess/Source/XCodeSourceCodeAccess/Private/XCodeSourceCodeAccessor.h
@@ -8,13 +8,13 @@ class FXCodeSourceCodeAccessor : public ISourceCodeAccessor
 {
 public:
 	/** ISourceCodeAccessor implementation */
-	virtual bool CanAccessSourceCode() const OVERRIDE;
-	virtual FName GetFName() const OVERRIDE;
-	virtual FText GetNameText() const OVERRIDE;
-	virtual FText GetDescriptionText() const OVERRIDE;
-	virtual bool OpenSolution() OVERRIDE;
-	virtual bool OpenFileAtLine(const FString& FullPath, int32 LineNumber, int32 ColumnNumber = 0) OVERRIDE;
-	virtual bool OpenSourceFiles(const TArray<FString>& AbsoluteSourcePaths) OVERRIDE;
-	virtual bool SaveAllOpenDocuments() const OVERRIDE;
-	virtual void Tick(const float DeltaTime) OVERRIDE;
+	virtual bool CanAccessSourceCode() const override;
+	virtual FName GetFName() const override;
+	virtual FText GetNameText() const override;
+	virtual FText GetDescriptionText() const override;
+	virtual bool OpenSolution() override;
+	virtual bool OpenFileAtLine(const FString& FullPath, int32 LineNumber, int32 ColumnNumber = 0) override;
+	virtual bool OpenSourceFiles(const TArray<FString>& AbsoluteSourcePaths) override;
+	virtual bool SaveAllOpenDocuments() const override;
+	virtual void Tick(const float DeltaTime) override;
 };
\ No newline at end of file
diff --git a/Engine/Plugins/Editor/EpicSurvey/Source/EpicSurvey/Private/EpicSurveyCommands.h b/Engine/Plugins/Editor/EpicSurvey/Source/EpicSurvey/Private/EpicSurveyCommands.h
--- a/Engine/Plugins/Editor/EpicSurvey/Source/EpicSurvey/Private/EpicSurveyCommands.h
+++ b/Engine/Plugins/Editor/EpicSurvey/Source/EpicSurvey/Private/EpicSurveyCommands.h
@@ -6,7 +6,7 @@ class FEpicSurveyCommands : public TCommands<FEpicSurveyCommands>
 public:
 	FEpicSurveyCommands();
 
-	virtual void RegisterCommands() OVERRIDE;
+	virtual void RegisterCommands() override;
 
 public:
 
diff --git a/Engine/Plugins/Editor/EpicSurvey/Source/EpicSurvey/Private/EpicSurveyModule.cpp b/Engine/Plugins/Editor/EpicSurvey/Source/EpicSurvey/Private/EpicSurveyModule.cpp
--- a/Engine/Plugins/Editor/EpicSurvey/Source/EpicSurvey/Private/EpicSurveyModule.cpp
+++ b/Engine/Plugins/Editor/EpicSurvey/Source/EpicSurvey/Private/EpicSurveyModule.cpp
@@ -10,10 +10,10 @@ class FEpicSurveyModule : public IEpicSurveyModule
 public:
 
 	/** IModuleInterface implementation */
-	virtual void StartupModule() OVERRIDE;
-	virtual void ShutdownModule() OVERRIDE;
+	virtual void StartupModule() override;
+	virtual void ShutdownModule() override;
 
-	virtual bool PromptSurvey( const FGuid& SurveyIdentifier ) OVERRIDE;
+	virtual bool PromptSurvey( const FGuid& SurveyIdentifier ) override;
 
 
 private:
diff --git a/Engine/Plugins/Editor/EpicSurvey/Source/EpicSurvey/Private/SSurvey.h b/Engine/Plugins/Editor/EpicSurvey/Source/EpicSurvey/Private/SSurvey.h
--- a/Engine/Plugins/Editor/EpicSurvey/Source/EpicSurvey/Private/SSurvey.h
+++ b/Engine/Plugins/Editor/EpicSurvey/Source/EpicSurvey/Private/SSurvey.h
@@ -15,9 +15,9 @@ public:
 
 	/** Widget constructor */
 	void Construct( const FArguments& Args, const TSharedRef< FEpicSurvey >& InEpicSurvey, const TSharedRef< FSurvey >& InSurvey );
-	virtual void Tick( const FGeometry& AllottedGeometry, const double InCurrentTime, const float InDeltaTime ) OVERRIDE;
+	virtual void Tick( const FGeometry& AllottedGeometry, const double InCurrentTime, const float InDeltaTime ) override;
 
-	virtual FReply OnMouseButtonDown( const FGeometry& MyGeometry, const FPointerEvent& MouseEvent ) OVERRIDE;
+	virtual FReply OnMouseButtonDown( const FGeometry& MyGeometry, const FPointerEvent& MouseEvent ) override;
 
 private:
 
diff --git a/Engine/Plugins/Editor/EpicSurvey/Source/EpicSurvey/Private/SSurveyNotification.h b/Engine/Plugins/Editor/EpicSurvey/Source/EpicSurvey/Private/SSurveyNotification.h
--- a/Engine/Plugins/Editor/EpicSurvey/Source/EpicSurvey/Private/SSurveyNotification.h
+++ b/Engine/Plugins/Editor/EpicSurvey/Source/EpicSurvey/Private/SSurveyNotification.h
@@ -36,7 +36,7 @@ public:
 		];
 	}
 
-	virtual void Tick( const FGeometry& AllottedGeometry, const double InCurrentTime, const float InDeltaTime ) OVERRIDE
+	virtual void Tick( const FGeometry& AllottedGeometry, const double InCurrentTime, const float InDeltaTime ) override
 	{
 		if ( this->IsHovered() )
 		{
diff --git a/Engine/Plugins/Editor/EpicSurvey/Source/EpicSurvey/Private/SurveyTitleCdnStorage.h b/Engine/Plugins/Editor/EpicSurvey/Source/EpicSurvey/Private/SurveyTitleCdnStorage.h
--- a/Engine/Plugins/Editor/EpicSurvey/Source/EpicSurvey/Private/SurveyTitleCdnStorage.h
+++ b/Engine/Plugins/Editor/EpicSurvey/Source/EpicSurvey/Private/SurveyTitleCdnStorage.h
@@ -14,19 +14,19 @@ public:
 
 public:
 
-	virtual bool GetFileContents(const FString& DLName, TArray<uint8>& FileContents) OVERRIDE;
+	virtual bool GetFileContents(const FString& DLName, TArray<uint8>& FileContents) override;
 
-	virtual bool ClearFiles() OVERRIDE;
+	virtual bool ClearFiles() override;
 
-	virtual bool ClearFile(const FString& DLName) OVERRIDE;
+	virtual bool ClearFile(const FString& DLName) override;
 
-	virtual bool EnumerateFiles() OVERRIDE;
+	virtual bool EnumerateFiles() override;
 
-	virtual bool EnumerateFiles(int32 Start, int32 Count) OVERRIDE;
+	virtual bool EnumerateFiles(int32 Start, int32 Count) override;
 
-	virtual void GetFileList(TArray<FCloudFileHeader>& InFileHeaders) OVERRIDE;
+	virtual void GetFileList(TArray<FCloudFileHeader>& InFileHeaders) override;
 
-	virtual bool ReadFile(const FString& DLName) OVERRIDE;
+	virtual bool ReadFile(const FString& DLName) override;
 
 private:
 	
diff --git a/Engine/Plugins/Editor/EpicSurvey/Source/EpicSurvey/Private/SurveyTitleLocalStorage.h b/Engine/Plugins/Editor/EpicSurvey/Source/EpicSurvey/Private/SurveyTitleLocalStorage.h
--- a/Engine/Plugins/Editor/EpicSurvey/Source/EpicSurvey/Private/SurveyTitleLocalStorage.h
+++ b/Engine/Plugins/Editor/EpicSurvey/Source/EpicSurvey/Private/SurveyTitleLocalStorage.h
@@ -10,19 +10,19 @@ public:
 
 public:
 
-	virtual bool GetFileContents(const FString& DLName, TArray<uint8>& FileContents) OVERRIDE;
+	virtual bool GetFileContents(const FString& DLName, TArray<uint8>& FileContents) override;
 
-	virtual bool ClearFiles() OVERRIDE;
+	virtual bool ClearFiles() override;
 
-	virtual bool ClearFile(const FString& DLName) OVERRIDE;
+	virtual bool ClearFile(const FString& DLName) override;
 
-	virtual bool EnumerateFiles() OVERRIDE;
+	virtual bool EnumerateFiles() override;
 
-	virtual bool EnumerateFiles(int32 Start, int32 Count) OVERRIDE;
+	virtual bool EnumerateFiles(int32 Start, int32 Count) override;
 
-	virtual void GetFileList(TArray<FCloudFileHeader>& InFileHeaders) OVERRIDE;
+	virtual void GetFileList(TArray<FCloudFileHeader>& InFileHeaders) override;
 
-	virtual bool ReadFile(const FString& DLName) OVERRIDE;
+	virtual bool ReadFile(const FString& DLName) override;
 
 private:
 
diff --git a/Engine/Plugins/Editor/IntroTutorials/Source/IntroTutorials/Classes/K2Node_TutorialExcerptComplete.h b/Engine/Plugins/Editor/IntroTutorials/Source/IntroTutorials/Classes/K2Node_TutorialExcerptComplete.h
--- a/Engine/Plugins/Editor/IntroTutorials/Source/IntroTutorials/Classes/K2Node_TutorialExcerptComplete.h
+++ b/Engine/Plugins/Editor/IntroTutorials/Source/IntroTutorials/Classes/K2Node_TutorialExcerptComplete.h
@@ -24,23 +24,23 @@ class UK2Node_TutorialExcerptComplete : public UK2Node_CallFunction
 	FString CachedExcerpt;
 
 	// Begin UK2Node_CallFunction interface
-	virtual void PostParameterPinCreated(UEdGraphPin *Pin) OVERRIDE;
+	virtual void PostParameterPinCreated(UEdGraphPin *Pin) override;
 	// End UK2Node_CallFunction interface
 
 	// Begin UEdGraphNode interface
-	virtual FText GetNodeTitle(ENodeTitleType::Type TitleType) const OVERRIDE;
-	virtual FString GetNodeNativeTitle(ENodeTitleType::Type TitleType) const OVERRIDE;
-	virtual void PostPlacedNewNode() OVERRIDE;
-	virtual void ValidateNodeDuringCompilation(class FCompilerResultsLog& MessageLog) const OVERRIDE;
+	virtual FText GetNodeTitle(ENodeTitleType::Type TitleType) const override;
+	virtual FString GetNodeNativeTitle(ENodeTitleType::Type TitleType) const override;
+	virtual void PostPlacedNewNode() override;
+	virtual void ValidateNodeDuringCompilation(class FCompilerResultsLog& MessageLog) const override;
 	// End UEdGraphNode interface
 
 	// Begin UK2Node interface
-	virtual void GetMenuEntries(FGraphContextMenuBuilder& ContextMenuBuilder) const OVERRIDE;
-	virtual void PinDefaultValueChanged(UEdGraphPin* Pin) OVERRIDE;
+	virtual void GetMenuEntries(FGraphContextMenuBuilder& ContextMenuBuilder) const override;
+	virtual void PinDefaultValueChanged(UEdGraphPin* Pin) override;
 	// End UK2Node interface
 
 	// Begin UObject interface
-	virtual void PostLoad() OVERRIDE;
+	virtual void PostLoad() override;
 	// End UObject interface
 
 	static bool IsExcerptInteractive(const FString& InExcerptName);
diff --git a/Engine/Plugins/Editor/PluginsEditor/Source/PluginsEditor/Private/SPluginCategories.h b/Engine/Plugins/Editor/PluginsEditor/Source/PluginsEditor/Private/SPluginCategories.h
--- a/Engine/Plugins/Editor/PluginsEditor/Source/PluginsEditor/Private/SPluginCategories.h
+++ b/Engine/Plugins/Editor/PluginsEditor/Source/PluginsEditor/Private/SPluginCategories.h
@@ -56,7 +56,7 @@ private:
 	void RebuildAndFilterCategoryTree();
 
 	/** SWidget overrides */
-	virtual void Tick( const FGeometry& AllottedGeometry, const double InCurrentTime, const float InDeltaTime ) OVERRIDE;
+	virtual void Tick( const FGeometry& AllottedGeometry, const double InCurrentTime, const float InDeltaTime ) override;
 
 
 private:
diff --git a/Engine/Plugins/Editor/PluginsEditor/Source/PluginsEditor/Private/SPluginList.h b/Engine/Plugins/Editor/PluginsEditor/Source/PluginsEditor/Private/SPluginList.h
--- a/Engine/Plugins/Editor/PluginsEditor/Source/PluginsEditor/Private/SPluginList.h
+++ b/Engine/Plugins/Editor/PluginsEditor/Source/PluginsEditor/Private/SPluginList.h
@@ -46,7 +46,7 @@ private:
 	void RebuildAndFilterPluginList();
 
 	/** SWidget overrides */
-	virtual void Tick( const FGeometry& AllottedGeometry, const double InCurrentTime, const float InDeltaTime ) OVERRIDE;
+	virtual void Tick( const FGeometry& AllottedGeometry, const double InCurrentTime, const float InDeltaTime ) override;
 
 
 private:
diff --git a/Engine/Plugins/Editor/SpeedTreeImporter/Source/SpeedTreeImporter/Classes/ReimportSpeedTreeFactory.h b/Engine/Plugins/Editor/SpeedTreeImporter/Source/SpeedTreeImporter/Classes/ReimportSpeedTreeFactory.h
--- a/Engine/Plugins/Editor/SpeedTreeImporter/Source/SpeedTreeImporter/Classes/ReimportSpeedTreeFactory.h
+++ b/Engine/Plugins/Editor/SpeedTreeImporter/Source/SpeedTreeImporter/Classes/ReimportSpeedTreeFactory.h
@@ -14,9 +14,9 @@ class UReimportSpeedTreeFactory : public USpeedTreeImportFactory, public FReimpo
 	GENERATED_UCLASS_BODY()
 
 	// Begin FReimportHandler interface
-	virtual bool CanReimport(UObject* Obj, TArray<FString>& OutFilenames) OVERRIDE;
-	virtual void SetReimportPaths(UObject* Obj, const TArray<FString>& NewReimportPaths) OVERRIDE;
-	virtual EReimportResult::Type Reimport(UObject* Obj) OVERRIDE;
+	virtual bool CanReimport(UObject* Obj, TArray<FString>& OutFilenames) override;
+	virtual void SetReimportPaths(UObject* Obj, const TArray<FString>& NewReimportPaths) override;
+	virtual EReimportResult::Type Reimport(UObject* Obj) override;
 	// End FReimportHandler interface
 };
 
diff --git a/Engine/Plugins/Editor/SpeedTreeImporter/Source/SpeedTreeImporter/Classes/SpeedTreeImportFactory.h b/Engine/Plugins/Editor/SpeedTreeImporter/Source/SpeedTreeImporter/Classes/SpeedTreeImportFactory.h
--- a/Engine/Plugins/Editor/SpeedTreeImporter/Source/SpeedTreeImporter/Classes/SpeedTreeImportFactory.h
+++ b/Engine/Plugins/Editor/SpeedTreeImporter/Source/SpeedTreeImporter/Classes/SpeedTreeImportFactory.h
@@ -13,11 +13,11 @@ class USpeedTreeImportFactory : public UFactory
 	GENERATED_UCLASS_BODY()
 
 	// Begin UFactory Interface
-	virtual FText GetDisplayName() const OVERRIDE;
+	virtual FText GetDisplayName() const override;
 #if WITH_SPEEDTREE
-	virtual bool DoesSupportClass(UClass * Class) OVERRIDE;
-	virtual UClass* ResolveSupportedClass() OVERRIDE;
-	virtual UObject* FactoryCreateBinary(UClass* InClass, UObject* InParent, FName InName, EObjectFlags Flags, UObject* Context, const TCHAR* Type, const uint8*& Buffer, const uint8* BufferEnd, FFeedbackContext* Warn) OVERRIDE;
+	virtual bool DoesSupportClass(UClass * Class) override;
+	virtual UClass* ResolveSupportedClass() override;
+	virtual UObject* FactoryCreateBinary(UClass* InClass, UObject* InParent, FName InName, EObjectFlags Flags, UObject* Context, const TCHAR* Type, const uint8*& Buffer, const uint8* BufferEnd, FFeedbackContext* Warn) override;
 #endif
 };
 
diff --git a/Engine/Plugins/Editor/SpeedTreeImporter/Source/SpeedTreeImporter/Private/SpeedTreeImporterModule.cpp b/Engine/Plugins/Editor/SpeedTreeImporter/Source/SpeedTreeImporter/Private/SpeedTreeImporterModule.cpp
--- a/Engine/Plugins/Editor/SpeedTreeImporter/Source/SpeedTreeImporter/Private/SpeedTreeImporterModule.cpp
+++ b/Engine/Plugins/Editor/SpeedTreeImporter/Source/SpeedTreeImporter/Private/SpeedTreeImporterModule.cpp
@@ -9,12 +9,12 @@
 class FSpeedTreeImporterModule : public ISpeedTreeImporter
 {
 public:
-	virtual void StartupModule() OVERRIDE
+	virtual void StartupModule() override
 	{
 	}
 
 
-	virtual void ShutdownModule() OVERRIDE
+	virtual void ShutdownModule() override
 	{
 	}
 
diff --git a/Engine/Plugins/Experimental/CodeView/Source/CodeView/Private/CodeViewPlugin.cpp b/Engine/Plugins/Experimental/CodeView/Source/CodeView/Private/CodeViewPlugin.cpp
--- a/Engine/Plugins/Experimental/CodeView/Source/CodeView/Private/CodeViewPlugin.cpp
+++ b/Engine/Plugins/Experimental/CodeView/Source/CodeView/Private/CodeViewPlugin.cpp
@@ -7,12 +7,12 @@ class FCodeViewPlugin : public IModuleInterface
 {
 public:
 	// IModuleInterface implementation
-	virtual void StartupModule() OVERRIDE
+	virtual void StartupModule() override
 	{
 		OnExtendActorDetails.AddRaw(this, &FCodeViewPlugin::AddCodeViewCategory);
 	}
 
-	virtual void ShutdownModule() OVERRIDE
+	virtual void ShutdownModule() override
 	{
 		OnExtendActorDetails.RemoveAll(this);
 	}
diff --git a/Engine/Plugins/Experimental/CodeView/Source/CodeView/Private/SCodeView.cpp b/Engine/Plugins/Experimental/CodeView/Source/CodeView/Private/SCodeView.cpp
--- a/Engine/Plugins/Experimental/CodeView/Source/CodeView/Private/SCodeView.cpp
+++ b/Engine/Plugins/Experimental/CodeView/Source/CodeView/Private/SCodeView.cpp
@@ -107,7 +107,7 @@ namespace CodeView
 
 
 		/** Overridden from SMultiColumnTableRow.  Generates a widget for this column of the tree row. */
-		virtual TSharedRef<SWidget> GenerateWidgetForColumn( const FName& ColumnName ) OVERRIDE
+		virtual TSharedRef<SWidget> GenerateWidgetForColumn( const FName& ColumnName ) override
 		{
 			// Create the widget for this item
 			auto ItemWidget = CodeViewWeak.Pin()->GenerateWidgetForItemAndColumn( Item, ColumnName );
diff --git a/Engine/Plugins/Experimental/CodeView/Source/CodeView/Private/SCodeView.h b/Engine/Plugins/Experimental/CodeView/Source/CodeView/Private/SCodeView.h
--- a/Engine/Plugins/Experimental/CodeView/Source/CodeView/Private/SCodeView.h
+++ b/Engine/Plugins/Experimental/CodeView/Source/CodeView/Private/SCodeView.h
@@ -46,12 +46,12 @@ namespace CodeView
 		: public FTreeItem
 	{
 	public:
-		virtual ETreeItemType::Type GetType() const OVERRIDE
+		virtual ETreeItemType::Type GetType() const override
 		{
 			return ETreeItemType::Class;
 		}
 
-		virtual const FString& GetAsString() const OVERRIDE
+		virtual const FString& GetAsString() const override
 		{
 			return ClassName;
 		}
@@ -72,12 +72,12 @@ namespace CodeView
 		: public FTreeItem
 	{
 	public:
-		virtual ETreeItemType::Type GetType() const OVERRIDE
+		virtual ETreeItemType::Type GetType() const override
 		{
 			return ETreeItemType::Function;
 		}
 
-		virtual const FString& GetAsString() const OVERRIDE
+		virtual const FString& GetAsString() const override
 		{
 			return FunctionName;
 		}
diff --git a/Engine/Plugins/Experimental/LiveEditor/Source/LiveEditor/Classes/K2Node_LiveEditObject.h b/Engine/Plugins/Experimental/LiveEditor/Source/LiveEditor/Classes/K2Node_LiveEditObject.h
--- a/Engine/Plugins/Experimental/LiveEditor/Source/LiveEditor/Classes/K2Node_LiveEditObject.h
+++ b/Engine/Plugins/Experimental/LiveEditor/Source/LiveEditor/Classes/K2Node_LiveEditObject.h
@@ -10,18 +10,18 @@ class UK2Node_LiveEditObject : public UK2Node
 
 #if WITH_EDITOR
 	// Begin UEdGraphNode interface.
-	virtual void AllocateDefaultPins() OVERRIDE;
-	virtual FText GetNodeTitle(ENodeTitleType::Type TitleType) const OVERRIDE;
-	virtual FString GetNodeNativeTitle(ENodeTitleType::Type TitleType) const OVERRIDE;
-	virtual void PinDefaultValueChanged(UEdGraphPin* Pin) OVERRIDE;
+	virtual void AllocateDefaultPins() override;
+	virtual FText GetNodeTitle(ENodeTitleType::Type TitleType) const override;
+	virtual FString GetNodeNativeTitle(ENodeTitleType::Type TitleType) const override;
+	virtual void PinDefaultValueChanged(UEdGraphPin* Pin) override;
 	// End UEdGraphNode interface.
 
 	// Begin UK2Node interface
-	virtual bool IsNodeSafeToIgnore() const OVERRIDE { return true; }
-	virtual void ReallocatePinsDuringReconstruction(TArray<UEdGraphPin*>& OldPins) OVERRIDE;
-	virtual class FNodeHandlingFunctor* CreateNodeHandler(class FKismetCompilerContext& CompilerContext) const OVERRIDE;
-	virtual void ExpandNode(class FKismetCompilerContext& CompilerContext, UEdGraph* SourceGraph) OVERRIDE;
-	virtual void GetNodeAttributes( TArray<TKeyValuePair<FString, FString>>& OutNodeAttributes ) const OVERRIDE;
+	virtual bool IsNodeSafeToIgnore() const override { return true; }
+	virtual void ReallocatePinsDuringReconstruction(TArray<UEdGraphPin*>& OldPins) override;
+	virtual class FNodeHandlingFunctor* CreateNodeHandler(class FKismetCompilerContext& CompilerContext) const override;
+	virtual void ExpandNode(class FKismetCompilerContext& CompilerContext, UEdGraph* SourceGraph) override;
+	virtual void GetNodeAttributes( TArray<TKeyValuePair<FString, FString>>& OutNodeAttributes ) const override;
 	// End UK2Node interface
 
 	/** Create new pins to show LiveEditable properties on archetype */
diff --git a/Engine/Plugins/Experimental/LiveEditor/Source/LiveEditor/Classes/LiveEditorBlueprint.h b/Engine/Plugins/Experimental/LiveEditor/Source/LiveEditor/Classes/LiveEditorBlueprint.h
--- a/Engine/Plugins/Experimental/LiveEditor/Source/LiveEditor/Classes/LiveEditorBlueprint.h
+++ b/Engine/Plugins/Experimental/LiveEditor/Source/LiveEditor/Classes/LiveEditorBlueprint.h
@@ -27,7 +27,7 @@ public:
 public:
 	void DoInit();
 
-	virtual class UWorld* GetWorld() const OVERRIDE;
+	virtual class UWorld* GetWorld() const override;
 
 private:
 	TMap< int32, TArray<FLiveEditorCheckpointData> > CheckpointMap;
diff --git a/Engine/Plugins/Experimental/LiveEditor/Source/LiveEditor/Private/LiveEditorConfigWindow.cpp b/Engine/Plugins/Experimental/LiveEditor/Source/LiveEditor/Private/LiveEditorConfigWindow.cpp
--- a/Engine/Plugins/Experimental/LiveEditor/Source/LiveEditor/Private/LiveEditorConfigWindow.cpp
+++ b/Engine/Plugins/Experimental/LiveEditor/Source/LiveEditor/Private/LiveEditorConfigWindow.cpp
@@ -189,7 +189,7 @@ public:
 	SLATE_BEGIN_ARGS( SLiveEditorDeviceData ) {}
 	SLATE_END_ARGS()
 
-	virtual void Tick( const FGeometry& AllottedGeometry, const double InCurrentTime, const float InDeltaTime ) OVERRIDE
+	virtual void Tick( const FGeometry& AllottedGeometry, const double InCurrentTime, const float InDeltaTime ) override
 	{
 		SCompoundWidget::Tick( AllottedGeometry, InCurrentTime, InDeltaTime );
 
diff --git a/Engine/Plugins/Experimental/LiveEditor/Source/LiveEditor/Private/LiveEditorManager.h b/Engine/Plugins/Experimental/LiveEditor/Source/LiveEditor/Private/LiveEditorManager.h
--- a/Engine/Plugins/Experimental/LiveEditor/Source/LiveEditor/Private/LiveEditorManager.h
+++ b/Engine/Plugins/Experimental/LiveEditor/Source/LiveEditor/Private/LiveEditorManager.h
@@ -42,7 +42,7 @@ public:
 	{
 		return bInitialized;
 	}
-	virtual TStatId GetStatId() const OVERRIDE
+	virtual TStatId GetStatId() const override
 	{
 		RETURN_QUICK_DECLARE_CYCLE_STAT(FLiveEditorManager, STATGROUP_Tickables);
 	}
diff --git a/Engine/Plugins/Experimental/LiveEditor/Source/LiveEditor/Private/SGraphPinLiveEditVar.h b/Engine/Plugins/Experimental/LiveEditor/Source/LiveEditor/Private/SGraphPinLiveEditVar.h
--- a/Engine/Plugins/Experimental/LiveEditor/Source/LiveEditor/Private/SGraphPinLiveEditVar.h
+++ b/Engine/Plugins/Experimental/LiveEditor/Source/LiveEditor/Private/SGraphPinLiveEditVar.h
@@ -19,7 +19,7 @@ protected:
 	 *
 	 *	@return Reference to the newly created widget object
 	 */
-	virtual TSharedRef<SWidget>	GetDefaultValueWidget() OVERRIDE;
+	virtual TSharedRef<SWidget>	GetDefaultValueWidget() override;
 
 private:
 
diff --git a/Engine/Plugins/Experimental/Paper2D/Source/Paper2D/Classes/PaperAnimatedRenderComponent.h b/Engine/Plugins/Experimental/Paper2D/Source/Paper2D/Classes/PaperAnimatedRenderComponent.h
--- a/Engine/Plugins/Experimental/Paper2D/Source/Paper2D/Classes/PaperAnimatedRenderComponent.h
+++ b/Engine/Plugins/Experimental/Paper2D/Source/Paper2D/Classes/PaperAnimatedRenderComponent.h
@@ -51,13 +51,13 @@ protected:
 	void CalculateCurrentFrame();
 public:
 	// UActorComponent interface
-	virtual void TickComponent(float DeltaTime, enum ELevelTick TickType, FActorComponentTickFunction* ThisTickFunction) OVERRIDE;
-	virtual void SendRenderDynamicData_Concurrent() OVERRIDE;
-	virtual const UObject* AdditionalStatObject() const OVERRIDE;
+	virtual void TickComponent(float DeltaTime, enum ELevelTick TickType, FActorComponentTickFunction* ThisTickFunction) override;
+	virtual void SendRenderDynamicData_Concurrent() override;
+	virtual const UObject* AdditionalStatObject() const override;
 	// End of UActorComponent interface
 
 	// UPrimitiveComponent interface
-	virtual FPrimitiveSceneProxy* CreateSceneProxy() OVERRIDE;
-	virtual FBoxSphereBounds CalcBounds(const FTransform & LocalToWorld) const OVERRIDE;
+	virtual FPrimitiveSceneProxy* CreateSceneProxy() override;
+	virtual FBoxSphereBounds CalcBounds(const FTransform & LocalToWorld) const override;
 	// End of UPrimitiveComponent interface
 };
diff --git a/Engine/Plugins/Experimental/Paper2D/Source/Paper2D/Classes/PaperFlipbookActor.h b/Engine/Plugins/Experimental/Paper2D/Source/Paper2D/Classes/PaperFlipbookActor.h
--- a/Engine/Plugins/Experimental/Paper2D/Source/Paper2D/Classes/PaperFlipbookActor.h
+++ b/Engine/Plugins/Experimental/Paper2D/Source/Paper2D/Classes/PaperFlipbookActor.h
@@ -15,7 +15,7 @@ class APaperFlipbookActor : public AActor
 
 	// AActor interface
 #if WITH_EDITOR
-	virtual bool GetReferencedContentObjects(TArray<UObject*>& Objects) const OVERRIDE;
+	virtual bool GetReferencedContentObjects(TArray<UObject*>& Objects) const override;
 #endif
 	// End of AActor interface
 };
diff --git a/Engine/Plugins/Experimental/Paper2D/Source/Paper2D/Classes/PaperRenderActor.h b/Engine/Plugins/Experimental/Paper2D/Source/Paper2D/Classes/PaperRenderActor.h
--- a/Engine/Plugins/Experimental/Paper2D/Source/Paper2D/Classes/PaperRenderActor.h
+++ b/Engine/Plugins/Experimental/Paper2D/Source/Paper2D/Classes/PaperRenderActor.h
@@ -15,7 +15,7 @@ class APaperRenderActor : public AActor
 
 	// AActor interface
 #if WITH_EDITOR
-	virtual bool GetReferencedContentObjects(TArray<UObject*>& Objects) const OVERRIDE;
+	virtual bool GetReferencedContentObjects(TArray<UObject*>& Objects) const override;
 #endif
 	// End of AActor interface
 };
diff --git a/Engine/Plugins/Experimental/Paper2D/Source/Paper2D/Classes/PaperRenderComponent.h b/Engine/Plugins/Experimental/Paper2D/Source/Paper2D/Classes/PaperRenderComponent.h
--- a/Engine/Plugins/Experimental/Paper2D/Source/Paper2D/Classes/PaperRenderComponent.h
+++ b/Engine/Plugins/Experimental/Paper2D/Source/Paper2D/Classes/PaperRenderComponent.h
@@ -41,25 +41,25 @@ public:
 public:
 	// UObject interface
 #if WITH_EDITOR
-	virtual void PostEditChangeProperty(FPropertyChangedEvent& PropertyChangedEvent) OVERRIDE;
+	virtual void PostEditChangeProperty(FPropertyChangedEvent& PropertyChangedEvent) override;
 #endif
 	// End of UObject interface
 
 	// UActorComponent interface
-	virtual void SendRenderDynamicData_Concurrent() OVERRIDE;
-	virtual const UObject* AdditionalStatObject() const OVERRIDE;
+	virtual void SendRenderDynamicData_Concurrent() override;
+	virtual const UObject* AdditionalStatObject() const override;
 	// End of UActorComponent interface
 
 	// USceneComponent interface
-	virtual bool HasAnySockets() const OVERRIDE;
-	virtual FTransform GetSocketTransform(FName InSocketName, ERelativeTransformSpace TransformSpace = RTS_World) const OVERRIDE;
-	virtual void QuerySupportedSockets(TArray<FComponentSocketDescription>& OutSockets) const OVERRIDE;
+	virtual bool HasAnySockets() const override;
+	virtual FTransform GetSocketTransform(FName InSocketName, ERelativeTransformSpace TransformSpace = RTS_World) const override;
+	virtual void QuerySupportedSockets(TArray<FComponentSocketDescription>& OutSockets) const override;
 	// End of USceneComponent interface
 
 	// UPrimitiveComponent interface
-	virtual FPrimitiveSceneProxy* CreateSceneProxy() OVERRIDE;
-	virtual FBoxSphereBounds CalcBounds(const FTransform & LocalToWorld) const OVERRIDE;
-	virtual class UBodySetup* GetBodySetup() OVERRIDE;
+	virtual FPrimitiveSceneProxy* CreateSceneProxy() override;
+	virtual FBoxSphereBounds CalcBounds(const FTransform & LocalToWorld) const override;
+	virtual class UBodySetup* GetBodySetup() override;
 	// End of UPrimitiveComponent interface
 
 protected:
diff --git a/Engine/Plugins/Experimental/Paper2D/Source/Paper2D/Classes/PaperSprite.h b/Engine/Plugins/Experimental/Paper2D/Source/Paper2D/Classes/PaperSprite.h
--- a/Engine/Plugins/Experimental/Paper2D/Source/Paper2D/Classes/PaperSprite.h
+++ b/Engine/Plugins/Experimental/Paper2D/Source/Paper2D/Classes/PaperSprite.h
@@ -125,10 +125,10 @@ public:
 
 #if WITH_EDITORONLY_DATA
 	// UObject interface
-	virtual void PostEditChangeProperty(FPropertyChangedEvent& PropertyChangedEvent) OVERRIDE;
-	virtual void Serialize(FArchive& Ar) OVERRIDE;
-	virtual void PostLoad() OVERRIDE;
-	virtual void GetAssetRegistryTags(TArray<FAssetRegistryTag>& OutTags) const OVERRIDE;
+	virtual void PostEditChangeProperty(FPropertyChangedEvent& PropertyChangedEvent) override;
+	virtual void Serialize(FArchive& Ar) override;
+	virtual void PostLoad() override;
+	virtual void GetAssetRegistryTags(TArray<FAssetRegistryTag>& OutTags) const override;
 	// End of UObject interface
 
 	FVector2D ConvertTextureSpaceToPivotSpace(FVector2D Input) const;
@@ -185,8 +185,8 @@ public:
 	void QuerySupportedSockets(TArray<FComponentSocketDescription>& OutSockets) const;
 
 	// IInterface_CollisionDataProvider interface
-	virtual bool GetPhysicsTriMeshData(struct FTriMeshCollisionData* CollisionData, bool InUseAllTriData) OVERRIDE;
-	virtual bool ContainsPhysicsTriMeshData(bool InUseAllTriData) const OVERRIDE;
+	virtual bool GetPhysicsTriMeshData(struct FTriMeshCollisionData* CollisionData, bool InUseAllTriData) override;
+	virtual bool ContainsPhysicsTriMeshData(bool InUseAllTriData) const override;
 	// End of IInterface_CollisionDataProvider
 
 
diff --git a/Engine/Plugins/Experimental/Paper2D/Source/Paper2D/Classes/PaperSpriteAtlas.h b/Engine/Plugins/Experimental/Paper2D/Source/Paper2D/Classes/PaperSpriteAtlas.h
--- a/Engine/Plugins/Experimental/Paper2D/Source/Paper2D/Classes/PaperSpriteAtlas.h
+++ b/Engine/Plugins/Experimental/Paper2D/Source/Paper2D/Classes/PaperSpriteAtlas.h
@@ -37,9 +37,9 @@ class UPaperSpriteAtlas : public UObject
 public:
 	// UObject interface
 #if WITH_EDITORONLY_DATA
-	virtual void PostInitProperties() OVERRIDE;
-	virtual void PostDuplicate(bool bDuplicateForPIE) OVERRIDE;
+	virtual void PostInitProperties() override;
+	virtual void PostDuplicate(bool bDuplicateForPIE) override;
 #endif
-	virtual void GetAssetRegistryTags(TArray<FAssetRegistryTag>& OutTags) const OVERRIDE;
+	virtual void GetAssetRegistryTags(TArray<FAssetRegistryTag>& OutTags) const override;
 	// End of UObject interface
 };
diff --git a/Engine/Plugins/Experimental/Paper2D/Source/Paper2D/Classes/PaperTileLayer.h b/Engine/Plugins/Experimental/Paper2D/Source/Paper2D/Classes/PaperTileLayer.h
--- a/Engine/Plugins/Experimental/Paper2D/Source/Paper2D/Classes/PaperTileLayer.h
+++ b/Engine/Plugins/Experimental/Paper2D/Source/Paper2D/Classes/PaperTileLayer.h
@@ -76,7 +76,7 @@ public:
 
 	// UObject interface
 #if WITH_EDITORONLY_DATA
-	virtual void PostEditChangeProperty(struct FPropertyChangedEvent& PropertyChangedEvent) OVERRIDE;
+	virtual void PostEditChangeProperty(struct FPropertyChangedEvent& PropertyChangedEvent) override;
 #endif
 	// End of UObject interface
 
diff --git a/Engine/Plugins/Experimental/Paper2D/Source/Paper2D/Classes/PaperTileMap.h b/Engine/Plugins/Experimental/Paper2D/Source/Paper2D/Classes/PaperTileMap.h
--- a/Engine/Plugins/Experimental/Paper2D/Source/Paper2D/Classes/PaperTileMap.h
+++ b/Engine/Plugins/Experimental/Paper2D/Source/Paper2D/Classes/PaperTileMap.h
@@ -53,7 +53,7 @@ class PAPER2D_API UPaperTileMap : public UDataAsset //@TODO: Just to make it eas
 public:
 	// UObject interface
 #if WITH_EDITOR
-	virtual void PostEditChangeProperty(FPropertyChangedEvent& PropertyChangedEvent) OVERRIDE;
+	virtual void PostEditChangeProperty(FPropertyChangedEvent& PropertyChangedEvent) override;
 #endif
 	// End of UObject interface
 
diff --git a/Engine/Plugins/Experimental/Paper2D/Source/Paper2D/Classes/PaperTileMapRenderComponent.h b/Engine/Plugins/Experimental/Paper2D/Source/Paper2D/Classes/PaperTileMapRenderComponent.h
--- a/Engine/Plugins/Experimental/Paper2D/Source/Paper2D/Classes/PaperTileMapRenderComponent.h
+++ b/Engine/Plugins/Experimental/Paper2D/Source/Paper2D/Classes/PaperTileMapRenderComponent.h
@@ -45,22 +45,22 @@ protected:
 
 public:
 	// UObject interface
-	virtual void PostInitProperties() OVERRIDE;
+	virtual void PostInitProperties() override;
 #if WITH_EDITORONLY_DATA
-	virtual void PostEditChangeProperty(FPropertyChangedEvent& PropertyChangedEvent) OVERRIDE;
-	virtual void Serialize(FArchive& Ar) OVERRIDE;
-	virtual void PostLoad() OVERRIDE;
+	virtual void PostEditChangeProperty(FPropertyChangedEvent& PropertyChangedEvent) override;
+	virtual void Serialize(FArchive& Ar) override;
+	virtual void PostLoad() override;
 #endif
 	// End of UObject interface
 
 	// UActorComponent interface
-	virtual void TickComponent(float DeltaTime, enum ELevelTick TickType, FActorComponentTickFunction* ThisTickFunction) OVERRIDE;
-	virtual const UObject* AdditionalStatObject() const OVERRIDE;
+	virtual void TickComponent(float DeltaTime, enum ELevelTick TickType, FActorComponentTickFunction* ThisTickFunction) override;
+	virtual const UObject* AdditionalStatObject() const override;
 	// End of UActorComponent interface
 
 	// UPrimitiveComponent interface
-	virtual FPrimitiveSceneProxy* CreateSceneProxy() OVERRIDE;
-	virtual FBoxSphereBounds CalcBounds(const FTransform& LocalToWorld) const OVERRIDE;
-	virtual class UBodySetup* GetBodySetup() OVERRIDE;
+	virtual FPrimitiveSceneProxy* CreateSceneProxy() override;
+	virtual FBoxSphereBounds CalcBounds(const FTransform& LocalToWorld) const override;
+	virtual class UBodySetup* GetBodySetup() override;
 	// End of UPrimitiveComponent interface
 };
diff --git a/Engine/Plugins/Experimental/Paper2D/Source/Paper2D/Private/PaperRenderSceneProxy.cpp b/Engine/Plugins/Experimental/Paper2D/Source/Paper2D/Private/PaperRenderSceneProxy.cpp
--- a/Engine/Plugins/Experimental/Paper2D/Source/Paper2D/Private/PaperRenderSceneProxy.cpp
+++ b/Engine/Plugins/Experimental/Paper2D/Source/Paper2D/Private/PaperRenderSceneProxy.cpp
@@ -34,7 +34,7 @@ struct FPaperSpriteVertex
 class FPaperSpriteVertexBuffer : public FVertexBuffer
 {
 public:
-	virtual void InitRHI() OVERRIDE
+	virtual void InitRHI() override
 	{
 		FRHIResourceCreateInfo CreateInfo;
 		VertexBufferRHI = RHICreateVertexBuffer(sizeof(FPaperSpriteVertex), BUF_Static, CreateInfo);
@@ -128,22 +128,22 @@ public:
 	}
 
 	// FMaterialRenderProxy interface.
-	virtual const FMaterial* GetMaterial(ERHIFeatureLevel::Type InFeatureLevel) const OVERRIDE
+	virtual const FMaterial* GetMaterial(ERHIFeatureLevel::Type InFeatureLevel) const override
 	{
 		return Parent->GetMaterial(InFeatureLevel);
 	}
 
-	virtual bool GetVectorValue(const FName ParameterName, FLinearColor* OutValue, const FMaterialRenderContext& Context) const OVERRIDE
+	virtual bool GetVectorValue(const FName ParameterName, FLinearColor* OutValue, const FMaterialRenderContext& Context) const override
 	{
 		return Parent->GetVectorValue(ParameterName, OutValue, Context);
 	}
 
-	virtual bool GetScalarValue(const FName ParameterName, float* OutValue, const FMaterialRenderContext& Context) const OVERRIDE
+	virtual bool GetScalarValue(const FName ParameterName, float* OutValue, const FMaterialRenderContext& Context) const override
 	{
 		return Parent->GetScalarValue(ParameterName, OutValue, Context);
 	}
 
-	virtual bool GetTextureValue(const FName ParameterName,const UTexture** OutValue, const FMaterialRenderContext& Context) const OVERRIDE
+	virtual bool GetTextureValue(const FName ParameterName,const UTexture** OutValue, const FMaterialRenderContext& Context) const override
 	{
 		if (ParameterName == TextureParameterName)
 		{
diff --git a/Engine/Plugins/Experimental/Paper2D/Source/Paper2D/Private/PaperRenderSceneProxy.h b/Engine/Plugins/Experimental/Paper2D/Source/Paper2D/Private/PaperRenderSceneProxy.h
--- a/Engine/Plugins/Experimental/Paper2D/Source/Paper2D/Private/PaperRenderSceneProxy.h
+++ b/Engine/Plugins/Experimental/Paper2D/Source/Paper2D/Private/PaperRenderSceneProxy.h
@@ -14,12 +14,12 @@ public:
 	~FPaperRenderSceneProxy();
 
 	// FPrimitiveSceneProxy interface.
-	virtual void DrawDynamicElements(FPrimitiveDrawInterface* PDI, const FSceneView* View) OVERRIDE;
-	virtual FPrimitiveViewRelevance GetViewRelevance(const FSceneView* View) OVERRIDE;
-	virtual void OnTransformChanged() OVERRIDE;
-	virtual uint32 GetMemoryFootprint() const OVERRIDE;
-	virtual bool CanBeOccluded() const OVERRIDE;
-	virtual void CreateRenderThreadResources() OVERRIDE;
+	virtual void DrawDynamicElements(FPrimitiveDrawInterface* PDI, const FSceneView* View) override;
+	virtual FPrimitiveViewRelevance GetViewRelevance(const FSceneView* View) override;
+	virtual void OnTransformChanged() override;
+	virtual uint32 GetMemoryFootprint() const override;
+	virtual bool CanBeOccluded() const override;
+	virtual void CreateRenderThreadResources() override;
 	// End of FPrimitiveSceneProxy interface.
 
 	void SetDrawCall_RenderThread(const FSpriteDrawCallRecord& NewDynamicData);
diff --git a/Engine/Plugins/Experimental/Paper2D/Source/Paper2D/Private/PaperTileMapRenderSceneProxy.h b/Engine/Plugins/Experimental/Paper2D/Source/Paper2D/Private/PaperTileMapRenderSceneProxy.h
--- a/Engine/Plugins/Experimental/Paper2D/Source/Paper2D/Private/PaperTileMapRenderSceneProxy.h
+++ b/Engine/Plugins/Experimental/Paper2D/Source/Paper2D/Private/PaperTileMapRenderSceneProxy.h
@@ -13,7 +13,7 @@ public:
 	FPaperTileMapRenderSceneProxy(const UPaperTileMapRenderComponent* InComponent);
 
 	// FPrimitiveSceneProxy interface.
-	virtual void DrawDynamicElements(FPrimitiveDrawInterface* PDI, const FSceneView* View) OVERRIDE;
+	virtual void DrawDynamicElements(FPrimitiveDrawInterface* PDI, const FSceneView* View) override;
 	// End of FPrimitiveSceneProxy interface.
 
 protected:
diff --git a/Engine/Plugins/Experimental/Paper2D/Source/Paper2D/Private/Rendering/PaperBatchComponent.h b/Engine/Plugins/Experimental/Paper2D/Source/Paper2D/Private/Rendering/PaperBatchComponent.h
--- a/Engine/Plugins/Experimental/Paper2D/Source/Paper2D/Private/Rendering/PaperBatchComponent.h
+++ b/Engine/Plugins/Experimental/Paper2D/Source/Paper2D/Private/Rendering/PaperBatchComponent.h
@@ -11,11 +11,11 @@ class UPaperBatchComponent : public UPrimitiveComponent
 	GENERATED_UCLASS_BODY()
 
 	// UActorComponent interface
-	virtual void TickComponent(float DeltaTime, enum ELevelTick TickType, FActorComponentTickFunction* ThisTickFunction) OVERRIDE;
+	virtual void TickComponent(float DeltaTime, enum ELevelTick TickType, FActorComponentTickFunction* ThisTickFunction) override;
 	// End of UActorComponent interface
 
 	// UPrimitiveComponent interface
-	virtual FPrimitiveSceneProxy* CreateSceneProxy() OVERRIDE;
-	virtual FBoxSphereBounds CalcBounds(const FTransform & LocalToWorld) const OVERRIDE;
+	virtual FPrimitiveSceneProxy* CreateSceneProxy() override;
+	virtual FBoxSphereBounds CalcBounds(const FTransform & LocalToWorld) const override;
 	// End of UPrimitiveComponent interface
 };
diff --git a/Engine/Plugins/Experimental/Paper2D/Source/Paper2D/Private/Rendering/PaperBatchSceneProxy.h b/Engine/Plugins/Experimental/Paper2D/Source/Paper2D/Private/Rendering/PaperBatchSceneProxy.h
--- a/Engine/Plugins/Experimental/Paper2D/Source/Paper2D/Private/Rendering/PaperBatchSceneProxy.h
+++ b/Engine/Plugins/Experimental/Paper2D/Source/Paper2D/Private/Rendering/PaperBatchSceneProxy.h
@@ -11,10 +11,10 @@ public:
 	FPaperBatchSceneProxy(const UPrimitiveComponent* InComponent);
 
  	// FPrimitiveSceneProxy interface.
- 	virtual void DrawDynamicElements(FPrimitiveDrawInterface* PDI, const FSceneView* View) OVERRIDE;
- 	virtual FPrimitiveViewRelevance GetViewRelevance(const FSceneView* View) OVERRIDE;
- 	virtual uint32 GetMemoryFootprint() const OVERRIDE;
- 	virtual bool CanBeOccluded() const OVERRIDE;
+ 	virtual void DrawDynamicElements(FPrimitiveDrawInterface* PDI, const FSceneView* View) override;
+ 	virtual FPrimitiveViewRelevance GetViewRelevance(const FSceneView* View) override;
+ 	virtual uint32 GetMemoryFootprint() const override;
+ 	virtual bool CanBeOccluded() const override;
  	// End of FPrimitiveSceneProxy interface.
 
 	void RegisterManagedProxy(class FPaperRenderSceneProxy* Proxy);
diff --git a/Engine/Plugins/Experimental/Paper2D/Source/Paper2DEditor/Classes/PaperFlipbookActorFactory.h b/Engine/Plugins/Experimental/Paper2D/Source/Paper2DEditor/Classes/PaperFlipbookActorFactory.h
--- a/Engine/Plugins/Experimental/Paper2D/Source/Paper2DEditor/Classes/PaperFlipbookActorFactory.h
+++ b/Engine/Plugins/Experimental/Paper2D/Source/Paper2DEditor/Classes/PaperFlipbookActorFactory.h
@@ -8,8 +8,8 @@ class UPaperFlipbookActorFactory : public UActorFactory
 	GENERATED_UCLASS_BODY()
 
 	// Begin UActorFactory Interface
-	virtual void PostSpawnActor(UObject* Asset, AActor* NewActor) OVERRIDE;
-	virtual void PostCreateBlueprint(UObject* Asset, AActor* CDO) OVERRIDE;
-	virtual bool CanCreateActorFrom(const FAssetData& AssetData, FText& OutErrorMsg) OVERRIDE;
+	virtual void PostSpawnActor(UObject* Asset, AActor* NewActor) override;
+	virtual void PostCreateBlueprint(UObject* Asset, AActor* CDO) override;
+	virtual bool CanCreateActorFrom(const FAssetData& AssetData, FText& OutErrorMsg) override;
 	// End UActorFactory Interface
 };
diff --git a/Engine/Plugins/Experimental/Paper2D/Source/Paper2DEditor/Classes/PaperFlipbookFactory.h b/Engine/Plugins/Experimental/Paper2D/Source/Paper2DEditor/Classes/PaperFlipbookFactory.h
--- a/Engine/Plugins/Experimental/Paper2D/Source/Paper2DEditor/Classes/PaperFlipbookFactory.h
+++ b/Engine/Plugins/Experimental/Paper2D/Source/Paper2DEditor/Classes/PaperFlipbookFactory.h
@@ -12,6 +12,6 @@ class UPaperFlipbookFactory : public UFactory
 	GENERATED_UCLASS_BODY()
 
 	// UFactory interface
-	virtual UObject* FactoryCreateNew(UClass* Class, UObject* InParent, FName Name, EObjectFlags Flags, UObject* Context, FFeedbackContext* Warn) OVERRIDE;
+	virtual UObject* FactoryCreateNew(UClass* Class, UObject* InParent, FName Name, EObjectFlags Flags, UObject* Context, FFeedbackContext* Warn) override;
 	// End of UFactory interface
 };
diff --git a/Engine/Plugins/Experimental/Paper2D/Source/Paper2DEditor/Classes/PaperFlipbookThumbnailRenderer.h b/Engine/Plugins/Experimental/Paper2D/Source/Paper2DEditor/Classes/PaperFlipbookThumbnailRenderer.h
--- a/Engine/Plugins/Experimental/Paper2D/Source/Paper2DEditor/Classes/PaperFlipbookThumbnailRenderer.h
+++ b/Engine/Plugins/Experimental/Paper2D/Source/Paper2DEditor/Classes/PaperFlipbookThumbnailRenderer.h
@@ -8,6 +8,6 @@ class UPaperFlipbookThumbnailRenderer : public UPaperSpriteThumbnailRenderer
 	GENERATED_UCLASS_BODY()
 
 	// UThumbnailRenderer interface
-	virtual void Draw(UObject* Object, int32 X, int32 Y, uint32 Width, uint32 Height, FRenderTarget*, FCanvas* Canvas) OVERRIDE;
+	virtual void Draw(UObject* Object, int32 X, int32 Y, uint32 Width, uint32 Height, FRenderTarget*, FCanvas* Canvas) override;
 	// End of UThumbnailRenderer interface
 };
diff --git a/Engine/Plugins/Experimental/Paper2D/Source/Paper2DEditor/Classes/PaperSpriteActorFactory.h b/Engine/Plugins/Experimental/Paper2D/Source/Paper2DEditor/Classes/PaperSpriteActorFactory.h
--- a/Engine/Plugins/Experimental/Paper2D/Source/Paper2DEditor/Classes/PaperSpriteActorFactory.h
+++ b/Engine/Plugins/Experimental/Paper2D/Source/Paper2DEditor/Classes/PaperSpriteActorFactory.h
@@ -8,8 +8,8 @@ class UPaperSpriteActorFactory : public UActorFactory
 	GENERATED_UCLASS_BODY()
 
 	// Begin UActorFactory Interface
-	virtual void PostSpawnActor(UObject* Asset, AActor* NewActor) OVERRIDE;
-	virtual void PostCreateBlueprint(UObject* Asset, AActor* CDO) OVERRIDE;
-	virtual bool CanCreateActorFrom(const FAssetData& AssetData, FText& OutErrorMsg) OVERRIDE;
+	virtual void PostSpawnActor(UObject* Asset, AActor* NewActor) override;
+	virtual void PostCreateBlueprint(UObject* Asset, AActor* CDO) override;
+	virtual bool CanCreateActorFrom(const FAssetData& AssetData, FText& OutErrorMsg) override;
 	// End UActorFactory Interface
 };
diff --git a/Engine/Plugins/Experimental/Paper2D/Source/Paper2DEditor/Classes/PaperSpriteFactory.h b/Engine/Plugins/Experimental/Paper2D/Source/Paper2DEditor/Classes/PaperSpriteFactory.h
--- a/Engine/Plugins/Experimental/Paper2D/Source/Paper2DEditor/Classes/PaperSpriteFactory.h
+++ b/Engine/Plugins/Experimental/Paper2D/Source/Paper2DEditor/Classes/PaperSpriteFactory.h
@@ -21,8 +21,8 @@ class UPaperSpriteFactory : public UFactory
 
 
 	// UFactory interface
-	virtual bool ConfigureProperties() OVERRIDE;
-	virtual UObject* FactoryCreateNew(UClass* Class, UObject* InParent, FName Name, EObjectFlags Flags, UObject* Context, FFeedbackContext* Warn) OVERRIDE;
+	virtual bool ConfigureProperties() override;
+	virtual UObject* FactoryCreateNew(UClass* Class, UObject* InParent, FName Name, EObjectFlags Flags, UObject* Context, FFeedbackContext* Warn) override;
 	// End of UFactory interface
 
 protected:
diff --git a/Engine/Plugins/Experimental/Paper2D/Source/Paper2DEditor/Classes/PaperSpriteThumbnailRenderer.h b/Engine/Plugins/Experimental/Paper2D/Source/Paper2DEditor/Classes/PaperSpriteThumbnailRenderer.h
--- a/Engine/Plugins/Experimental/Paper2D/Source/Paper2DEditor/Classes/PaperSpriteThumbnailRenderer.h
+++ b/Engine/Plugins/Experimental/Paper2D/Source/Paper2DEditor/Classes/PaperSpriteThumbnailRenderer.h
@@ -8,7 +8,7 @@ class UPaperSpriteThumbnailRenderer : public UDefaultSizedThumbnailRenderer
 	GENERATED_UCLASS_BODY()
 
 	// UThumbnailRenderer interface
-	virtual void Draw(UObject* Object, int32 X, int32 Y, uint32 Width, uint32 Height, FRenderTarget*, FCanvas* Canvas) OVERRIDE;
+	virtual void Draw(UObject* Object, int32 X, int32 Y, uint32 Width, uint32 Height, FRenderTarget*, FCanvas* Canvas) override;
 	// End of UThumbnailRenderer interface
 
 protected:
diff --git a/Engine/Plugins/Experimental/Paper2D/Source/Paper2DEditor/Private/Atlasing/AtlasAssetTypeActions.h b/Engine/Plugins/Experimental/Paper2D/Source/Paper2DEditor/Private/Atlasing/AtlasAssetTypeActions.h
--- a/Engine/Plugins/Experimental/Paper2D/Source/Paper2DEditor/Private/Atlasing/AtlasAssetTypeActions.h
+++ b/Engine/Plugins/Experimental/Paper2D/Source/Paper2DEditor/Private/Atlasing/AtlasAssetTypeActions.h
@@ -8,11 +8,11 @@ class FAtlasAssetTypeActions : public FAssetTypeActions_Base
 {
 public:
 	// IAssetTypeActions interface
-	virtual FText GetName() const OVERRIDE;
-	virtual FColor GetTypeColor() const OVERRIDE;
-	virtual UClass* GetSupportedClass() const OVERRIDE;
-	virtual void OpenAssetEditor(const TArray<UObject*>& InObjects, TSharedPtr<class IToolkitHost> EditWithinLevelEditor = TSharedPtr<IToolkitHost>()) OVERRIDE;
-	virtual uint32 GetCategories() OVERRIDE;
-	virtual FText GetAssetDescription(const FAssetData& AssetData) const OVERRIDE;
+	virtual FText GetName() const override;
+	virtual FColor GetTypeColor() const override;
+	virtual UClass* GetSupportedClass() const override;
+	virtual void OpenAssetEditor(const TArray<UObject*>& InObjects, TSharedPtr<class IToolkitHost> EditWithinLevelEditor = TSharedPtr<IToolkitHost>()) override;
+	virtual uint32 GetCategories() override;
+	virtual FText GetAssetDescription(const FAssetData& AssetData) const override;
 	// End of IAssetTypeActions interface
 };
\ No newline at end of file
diff --git a/Engine/Plugins/Experimental/Paper2D/Source/Paper2DEditor/Private/Atlasing/PaperAtlasGenerator.cpp b/Engine/Plugins/Experimental/Paper2D/Source/Paper2DEditor/Private/Atlasing/PaperAtlasGenerator.cpp
--- a/Engine/Plugins/Experimental/Paper2D/Source/Paper2DEditor/Private/Atlasing/PaperAtlasGenerator.cpp
+++ b/Engine/Plugins/Experimental/Paper2D/Source/Paper2DEditor/Private/Atlasing/PaperAtlasGenerator.cpp
@@ -32,7 +32,7 @@ public:
 		//@TODO: check(StrideBytes >= 4*Width);
 	}
 
-	virtual void ConditionalUpdateTexture() OVERRIDE {}
+	virtual void ConditionalUpdateTexture() override {}
 
 	void CopyToTexture(UTexture2D* Texture)
 	{
diff --git a/Engine/Plugins/Experimental/Paper2D/Source/Paper2DEditor/Private/Atlasing/PaperSpriteAtlasFactory.h b/Engine/Plugins/Experimental/Paper2D/Source/Paper2DEditor/Private/Atlasing/PaperSpriteAtlasFactory.h
--- a/Engine/Plugins/Experimental/Paper2D/Source/Paper2DEditor/Private/Atlasing/PaperSpriteAtlasFactory.h
+++ b/Engine/Plugins/Experimental/Paper2D/Source/Paper2DEditor/Private/Atlasing/PaperSpriteAtlasFactory.h
@@ -10,7 +10,7 @@ class UPaperSpriteAtlasFactory : public UFactory
 	GENERATED_UCLASS_BODY()
 
 	// UFactory interface
-	virtual UObject* FactoryCreateNew(UClass* Class, UObject* InParent, FName Name, EObjectFlags Flags, UObject* Context, FFeedbackContext* Warn) OVERRIDE;
+	virtual UObject* FactoryCreateNew(UClass* Class, UObject* InParent, FName Name, EObjectFlags Flags, UObject* Context, FFeedbackContext* Warn) override;
 	// End of UFactory interface
 };
 
diff --git a/Engine/Plugins/Experimental/Paper2D/Source/Paper2DEditor/Private/ContentBrowserExtensions/ContentBrowserExtensions.cpp b/Engine/Plugins/Experimental/Paper2D/Source/Paper2DEditor/Private/ContentBrowserExtensions/ContentBrowserExtensions.cpp
--- a/Engine/Plugins/Experimental/Paper2D/Source/Paper2DEditor/Private/ContentBrowserExtensions/ContentBrowserExtensions.cpp
+++ b/Engine/Plugins/Experimental/Paper2D/Source/Paper2DEditor/Private/ContentBrowserExtensions/ContentBrowserExtensions.cpp
@@ -125,7 +125,7 @@ struct FCreateSpriteFromTextureExtension : public FContentBrowserSelectedAssetEx
 		}
 	}
 
-	virtual void Execute() OVERRIDE
+	virtual void Execute() override
 	{
 		// Create sprites for any selected textures
 		TArray<UTexture2D*> Textures;
@@ -144,7 +144,7 @@ struct FCreateSpriteFromTextureExtension : public FContentBrowserSelectedAssetEx
 
 struct FConfigureTexturesForSpriteUsageExtension : public FContentBrowserSelectedAssetExtensionBase
 {
-	virtual void Execute() OVERRIDE
+	virtual void Execute() override
 	{
 		// Change the compression settings and trigger a recompress
 		for (auto AssetIt = SelectedAssets.CreateConstIterator(); AssetIt; ++AssetIt)
diff --git a/Engine/Plugins/Experimental/Paper2D/Source/Paper2DEditor/Private/FlipbookAssetTypeActions.h b/Engine/Plugins/Experimental/Paper2D/Source/Paper2DEditor/Private/FlipbookAssetTypeActions.h
--- a/Engine/Plugins/Experimental/Paper2D/Source/Paper2DEditor/Private/FlipbookAssetTypeActions.h
+++ b/Engine/Plugins/Experimental/Paper2D/Source/Paper2DEditor/Private/FlipbookAssetTypeActions.h
@@ -8,10 +8,10 @@ class FFlipbookAssetTypeActions : public FAssetTypeActions_Base
 {
 public:
 	// IAssetTypeActions interface
-	virtual FText GetName() const OVERRIDE;
-	virtual FColor GetTypeColor() const OVERRIDE;
-	virtual UClass* GetSupportedClass() const OVERRIDE;
-	virtual void OpenAssetEditor(const TArray<UObject*>& InObjects, TSharedPtr<class IToolkitHost> EditWithinLevelEditor = TSharedPtr<IToolkitHost>()) OVERRIDE;
-	virtual uint32 GetCategories() OVERRIDE;
+	virtual FText GetName() const override;
+	virtual FColor GetTypeColor() const override;
+	virtual UClass* GetSupportedClass() const override;
+	virtual void OpenAssetEditor(const TArray<UObject*>& InObjects, TSharedPtr<class IToolkitHost> EditWithinLevelEditor = TSharedPtr<IToolkitHost>()) override;
+	virtual uint32 GetCategories() override;
 	// End of IAssetTypeActions interface
 };
\ No newline at end of file
diff --git a/Engine/Plugins/Experimental/Paper2D/Source/Paper2DEditor/Private/FlipbookEditor/FlipbookEditor.cpp b/Engine/Plugins/Experimental/Paper2D/Source/Paper2DEditor/Private/FlipbookEditor/FlipbookEditor.cpp
--- a/Engine/Plugins/Experimental/Paper2D/Source/Paper2DEditor/Private/FlipbookEditor/FlipbookEditor.cpp
+++ b/Engine/Plugins/Experimental/Paper2D/Source/Paper2DEditor/Private/FlipbookEditor/FlipbookEditor.cpp
@@ -54,9 +54,9 @@ public:
 	void Construct(const FArguments& InArgs);
 
 	// SEditorViewport interface
-	virtual void BindCommands() OVERRIDE;
-	virtual TSharedRef<FEditorViewportClient> MakeEditorViewportClient() OVERRIDE;
-	virtual TSharedPtr<SWidget> MakeViewportToolbar() OVERRIDE;
+	virtual void BindCommands() override;
+	virtual TSharedRef<FEditorViewportClient> MakeEditorViewportClient() override;
+	virtual TSharedPtr<SWidget> MakeViewportToolbar() override;
 	// End of SEditorViewport interface
 
 private:
@@ -152,12 +152,12 @@ public:
 	}
 
 	// SSingleObjectDetailsPanel interface
-	virtual UObject* GetObjectToObserve() const OVERRIDE
+	virtual UObject* GetObjectToObserve() const override
 	{
 		return FlipbookEditorPtr.Pin()->GetFlipbookBeingEdited();
 	}
 
-	virtual TSharedRef<SWidget> PopulateSlot(TSharedRef<SWidget> PropertyEditorWidget) OVERRIDE
+	virtual TSharedRef<SWidget> PopulateSlot(TSharedRef<SWidget> PropertyEditorWidget) override
 	{
 		return SNew(SVerticalBox)
 			+SVerticalBox::Slot()
diff --git a/Engine/Plugins/Experimental/Paper2D/Source/Paper2DEditor/Private/FlipbookEditor/FlipbookEditor.h b/Engine/Plugins/Experimental/Paper2D/Source/Paper2DEditor/Private/FlipbookEditor/FlipbookEditor.h
--- a/Engine/Plugins/Experimental/Paper2D/Source/Paper2DEditor/Private/FlipbookEditor/FlipbookEditor.h
+++ b/Engine/Plugins/Experimental/Paper2D/Source/Paper2DEditor/Private/FlipbookEditor/FlipbookEditor.h
@@ -12,20 +12,20 @@ class FFlipbookEditor : public FAssetEditorToolkit, public FGCObject
 {
 public:
 	// IToolkit interface
-	virtual void RegisterTabSpawners(const TSharedRef<class FTabManager>& TabManager) OVERRIDE;
-	virtual void UnregisterTabSpawners(const TSharedRef<class FTabManager>& TabManager) OVERRIDE;
+	virtual void RegisterTabSpawners(const TSharedRef<class FTabManager>& TabManager) override;
+	virtual void UnregisterTabSpawners(const TSharedRef<class FTabManager>& TabManager) override;
 	// End of IToolkit interface
 
 	// FAssetEditorToolkit
-	virtual FName GetToolkitFName() const OVERRIDE;
-	virtual FText GetBaseToolkitName() const OVERRIDE;
-	virtual FText GetToolkitName() const OVERRIDE;
-	virtual FLinearColor GetWorldCentricTabColorScale() const OVERRIDE;
-	virtual FString GetWorldCentricTabPrefix() const OVERRIDE;
+	virtual FName GetToolkitFName() const override;
+	virtual FText GetBaseToolkitName() const override;
+	virtual FText GetToolkitName() const override;
+	virtual FLinearColor GetWorldCentricTabColorScale() const override;
+	virtual FString GetWorldCentricTabPrefix() const override;
 	// End of FAssetEditorToolkit
 
 	// FSerializableObject interface
-	virtual void AddReferencedObjects(FReferenceCollector& Collector) OVERRIDE;
+	virtual void AddReferencedObjects(FReferenceCollector& Collector) override;
 	// End of FSerializableObject interface
 
 public:
diff --git a/Engine/Plugins/Experimental/Paper2D/Source/Paper2DEditor/Private/FlipbookEditor/FlipbookEditorCommands.h b/Engine/Plugins/Experimental/Paper2D/Source/Paper2DEditor/Private/FlipbookEditor/FlipbookEditorCommands.h
--- a/Engine/Plugins/Experimental/Paper2D/Source/Paper2DEditor/Private/FlipbookEditor/FlipbookEditorCommands.h
+++ b/Engine/Plugins/Experimental/Paper2D/Source/Paper2DEditor/Private/FlipbookEditor/FlipbookEditorCommands.h
@@ -16,7 +16,7 @@ public:
 	}
 
 	// TCommand<> interface
-	virtual void RegisterCommands() OVERRIDE;
+	virtual void RegisterCommands() override;
 	// End of TCommand<> interface
 
 public:
diff --git a/Engine/Plugins/Experimental/Paper2D/Source/Paper2DEditor/Private/FlipbookEditor/FlipbookEditorViewportClient.h b/Engine/Plugins/Experimental/Paper2D/Source/Paper2DEditor/Private/FlipbookEditor/FlipbookEditorViewportClient.h
--- a/Engine/Plugins/Experimental/Paper2D/Source/Paper2DEditor/Private/FlipbookEditor/FlipbookEditorViewportClient.h
+++ b/Engine/Plugins/Experimental/Paper2D/Source/Paper2DEditor/Private/FlipbookEditor/FlipbookEditorViewportClient.h
@@ -20,14 +20,14 @@ public:
 	FFlipbookEditorViewportClient(const TAttribute<class UPaperFlipbook*>& InFlipbookBeingEdited, const TAttribute<float>& InPlayTime);
 
 	// FViewportClient interface
-	virtual void Draw(FViewport* Viewport, FCanvas* Canvas) OVERRIDE;
+	virtual void Draw(FViewport* Viewport, FCanvas* Canvas) override;
 	virtual void Draw(const FSceneView* View, FPrimitiveDrawInterface* PDI);
-	virtual void DrawCanvas(FViewport& InViewport, FSceneView& View, FCanvas& Canvas) OVERRIDE;
-	virtual void Tick(float DeltaSeconds) OVERRIDE;
+	virtual void DrawCanvas(FViewport& InViewport, FSceneView& View, FCanvas& Canvas) override;
+	virtual void Tick(float DeltaSeconds) override;
 	// End of FViewportClient interface
 
 	// FEditorViewportClient interface
-	virtual bool InputKey(FViewport* Viewport, int32 ControllerId, FKey Key, EInputEvent Event, float AmountDepressed, bool bGamepad) OVERRIDE;
+	virtual bool InputKey(FViewport* Viewport, int32 ControllerId, FKey Key, EInputEvent Event, float AmountDepressed, bool bGamepad) override;
 	// End of FEditorViewportClient interface
 
 	void ToggleShowPivot() { bShowPivot = !bShowPivot; Invalidate(); }
diff --git a/Engine/Plugins/Experimental/Paper2D/Source/Paper2DEditor/Private/FlipbookEditor/SFlipbookTimeline.cpp b/Engine/Plugins/Experimental/Paper2D/Source/Paper2DEditor/Private/FlipbookEditor/SFlipbookTimeline.cpp
--- a/Engine/Plugins/Experimental/Paper2D/Source/Paper2DEditor/Private/FlipbookEditor/SFlipbookTimeline.cpp
+++ b/Engine/Plugins/Experimental/Paper2D/Source/Paper2DEditor/Private/FlipbookEditor/SFlipbookTimeline.cpp
@@ -32,7 +32,7 @@ public:
 		SImage::Construct( SImage::FArguments() );
 	}
 
-	virtual FReply OnMouseButtonDown( const FGeometry& MyGeometry, const FPointerEvent& MouseEvent ) OVERRIDE
+	virtual FReply OnMouseButtonDown( const FGeometry& MyGeometry, const FPointerEvent& MouseEvent ) override
 	{
 		if ( MouseEvent.GetEffectingButton() == EKeys::LeftMouseButton )
 		{
@@ -47,7 +47,7 @@ public:
 		}
 	}
 
-	virtual FReply OnMouseButtonUp( const FGeometry& MyGeometry, const FPointerEvent& MouseEvent ) OVERRIDE
+	virtual FReply OnMouseButtonUp( const FGeometry& MyGeometry, const FPointerEvent& MouseEvent ) override
 	{
 		if ( MouseEvent.GetEffectingButton() == EKeys::LeftMouseButton && this->HasMouseCapture() )
 		{
@@ -82,7 +82,7 @@ public:
 		}
 	}
 
-	virtual FReply OnMouseMove( const FGeometry& MyGeometry, const FPointerEvent& MouseEvent ) OVERRIDE
+	virtual FReply OnMouseMove( const FGeometry& MyGeometry, const FPointerEvent& MouseEvent ) override
 	{
 		if ( this->HasMouseCapture() )
 		{
@@ -118,7 +118,7 @@ public:
 		return FReply::Unhandled();
 	}
 
-	virtual FCursorReply OnCursorQuery( const FGeometry& MyGeometry, const FPointerEvent& CursorEvent ) const OVERRIDE
+	virtual FCursorReply OnCursorQuery( const FGeometry& MyGeometry, const FPointerEvent& CursorEvent ) const override
 	{
 		return bDragging ? 
 			FCursorReply::Cursor( EMouseCursor::None ) :
@@ -165,7 +165,7 @@ public:
 		Rebuild();
 	}
 
-	void Tick( const FGeometry& AllottedGeometry, const double InCurrentTime, const float InDeltaTime ) OVERRIDE
+	void Tick( const FGeometry& AllottedGeometry, const double InCurrentTime, const float InDeltaTime ) override
 	{
 		SCompoundWidget::Tick(AllottedGeometry, InCurrentTime, InDeltaTime);
 
@@ -294,7 +294,7 @@ public:
 		Rebuild();
 	}
 
-	void Tick( const FGeometry& AllottedGeometry, const double InCurrentTime, const float InDeltaTime ) OVERRIDE
+	void Tick( const FGeometry& AllottedGeometry, const double InCurrentTime, const float InDeltaTime ) override
 	{
 		SCompoundWidget::Tick(AllottedGeometry, InCurrentTime, InDeltaTime);
 
diff --git a/Engine/Plugins/Experimental/Paper2D/Source/Paper2DEditor/Private/Paper2DEditorModule.cpp b/Engine/Plugins/Experimental/Paper2D/Source/Paper2DEditor/Private/Paper2DEditorModule.cpp
--- a/Engine/Plugins/Experimental/Paper2D/Source/Paper2DEditor/Private/Paper2DEditorModule.cpp
+++ b/Engine/Plugins/Experimental/Paper2D/Source/Paper2DEditor/Private/Paper2DEditorModule.cpp
@@ -40,11 +40,11 @@ class FPaper2DEditor : public IPaper2DEditorModule
 {
 public:
 	// IPaper2DEditorModule interface
-	virtual TSharedPtr<FExtensibilityManager> GetSpriteEditorMenuExtensibilityManager() OVERRIDE { return SpriteEditor_MenuExtensibilityManager; }
-	virtual TSharedPtr<FExtensibilityManager> GetSpriteEditorToolBarExtensibilityManager() OVERRIDE { return SpriteEditor_ToolBarExtensibilityManager; }
+	virtual TSharedPtr<FExtensibilityManager> GetSpriteEditorMenuExtensibilityManager() override { return SpriteEditor_MenuExtensibilityManager; }
+	virtual TSharedPtr<FExtensibilityManager> GetSpriteEditorToolBarExtensibilityManager() override { return SpriteEditor_ToolBarExtensibilityManager; }
 
-	virtual TSharedPtr<FExtensibilityManager> GetFlipbookEditorMenuExtensibilityManager() OVERRIDE { return FlipbookEditor_MenuExtensibilityManager; }
-	virtual TSharedPtr<FExtensibilityManager> GetFlipbookEditorToolBarExtensibilityManager() OVERRIDE { return FlipbookEditor_ToolBarExtensibilityManager; }
+	virtual TSharedPtr<FExtensibilityManager> GetFlipbookEditorMenuExtensibilityManager() override { return FlipbookEditor_MenuExtensibilityManager; }
+	virtual TSharedPtr<FExtensibilityManager> GetFlipbookEditorToolBarExtensibilityManager() override { return FlipbookEditor_ToolBarExtensibilityManager; }
 	// End of IPaper2DEditorModule
 
 private:
@@ -65,7 +65,7 @@ private:
 	FCoreDelegates::FOnObjectPropertyChanged::FDelegate OnPropertyChangedHandle;
 
 public:
-	virtual void StartupModule() OVERRIDE
+	virtual void StartupModule() override
 	{
 		SpriteEditor_MenuExtensibilityManager = MakeShareable(new FExtensibilityManager);
 		SpriteEditor_ToolBarExtensibilityManager = MakeShareable(new FExtensibilityManager);
@@ -123,7 +123,7 @@ public:
 		RegisterSettings();
 	}
 
-	virtual void ShutdownModule() OVERRIDE
+	virtual void ShutdownModule() override
 	{
 		SpriteEditor_MenuExtensibilityManager.Reset();
 		SpriteEditor_ToolBarExtensibilityManager.Reset();
diff --git a/Engine/Plugins/Experimental/Paper2D/Source/Paper2DEditor/Private/PaperEditorViewportClient.h b/Engine/Plugins/Experimental/Paper2D/Source/Paper2DEditor/Private/PaperEditorViewportClient.h
--- a/Engine/Plugins/Experimental/Paper2D/Source/Paper2DEditor/Private/PaperEditorViewportClient.h
+++ b/Engine/Plugins/Experimental/Paper2D/Source/Paper2DEditor/Private/PaperEditorViewportClient.h
@@ -22,11 +22,11 @@ public:
 	~FPaperEditorViewportClient();
 
 	// FViewportClient interface
-	virtual void Draw(FViewport* Viewport, FCanvas* Canvas) OVERRIDE;
+	virtual void Draw(FViewport* Viewport, FCanvas* Canvas) override;
 	// End of FViewportClient interface
 
 	// FSerializableObject interface
-	virtual void AddReferencedObjects(FReferenceCollector& Collector) OVERRIDE;
+	virtual void AddReferencedObjects(FReferenceCollector& Collector) override;
 	// End of FSerializableObject interface
 
 	/** Modifies the checkerboard texture's data */
diff --git a/Engine/Plugins/Experimental/Paper2D/Source/Paper2DEditor/Private/PaperFlipbookAssetBroker.h b/Engine/Plugins/Experimental/Paper2D/Source/Paper2DEditor/Private/PaperFlipbookAssetBroker.h
--- a/Engine/Plugins/Experimental/Paper2D/Source/Paper2DEditor/Private/PaperFlipbookAssetBroker.h
+++ b/Engine/Plugins/Experimental/Paper2D/Source/Paper2DEditor/Private/PaperFlipbookAssetBroker.h
@@ -10,12 +10,12 @@
 class FPaperFlipbookAssetBroker : public IComponentAssetBroker
 {
 public:
-	UClass* GetSupportedAssetClass() OVERRIDE
+	UClass* GetSupportedAssetClass() override
 	{
 		return UPaperFlipbook::StaticClass();
 	}
 
-	virtual bool AssignAssetToComponent(UActorComponent* InComponent, UObject* InAsset) OVERRIDE
+	virtual bool AssignAssetToComponent(UActorComponent* InComponent, UObject* InAsset) override
 	{
 		if (UPaperAnimatedRenderComponent* RenderComp = Cast<UPaperAnimatedRenderComponent>(InComponent))
 		{
@@ -31,7 +31,7 @@ public:
 		return false;
 	}
 
-	virtual UObject* GetAssetFromComponent(UActorComponent* InComponent) OVERRIDE
+	virtual UObject* GetAssetFromComponent(UActorComponent* InComponent) override
 	{
 		if (UPaperAnimatedRenderComponent* RenderComp = Cast<UPaperAnimatedRenderComponent>(InComponent))
 		{
diff --git a/Engine/Plugins/Experimental/Paper2D/Source/Paper2DEditor/Private/PaperSpriteAssetBroker.h b/Engine/Plugins/Experimental/Paper2D/Source/Paper2DEditor/Private/PaperSpriteAssetBroker.h
--- a/Engine/Plugins/Experimental/Paper2D/Source/Paper2DEditor/Private/PaperSpriteAssetBroker.h
+++ b/Engine/Plugins/Experimental/Paper2D/Source/Paper2DEditor/Private/PaperSpriteAssetBroker.h
@@ -10,12 +10,12 @@
 class FPaperSpriteAssetBroker : public IComponentAssetBroker
 {
 public:
-	UClass* GetSupportedAssetClass() OVERRIDE
+	UClass* GetSupportedAssetClass() override
 	{
 		return UPaperSprite::StaticClass();
 	}
 
-	virtual bool AssignAssetToComponent(UActorComponent* InComponent, UObject* InAsset) OVERRIDE
+	virtual bool AssignAssetToComponent(UActorComponent* InComponent, UObject* InAsset) override
 	{
 		if (UPaperRenderComponent* RenderComp = Cast<UPaperRenderComponent>(InComponent))
 		{
@@ -31,7 +31,7 @@ public:
 		return false;
 	}
 
-	virtual UObject* GetAssetFromComponent(UActorComponent* InComponent) OVERRIDE
+	virtual UObject* GetAssetFromComponent(UActorComponent* InComponent) override
 	{
 		if (UPaperRenderComponent* RenderComp = Cast<UPaperRenderComponent>(InComponent))
 		{
diff --git a/Engine/Plugins/Experimental/Paper2D/Source/Paper2DEditor/Private/SPaperEditorViewport.h b/Engine/Plugins/Experimental/Paper2D/Source/Paper2DEditor/Private/SPaperEditorViewport.h
--- a/Engine/Plugins/Experimental/Paper2D/Source/Paper2DEditor/Private/SPaperEditorViewport.h
+++ b/Engine/Plugins/Experimental/Paper2D/Source/Paper2DEditor/Private/SPaperEditorViewport.h
@@ -80,13 +80,13 @@ public:
 	~SPaperEditorViewport();
 
 	// SWidget interface
-	virtual void Tick(const FGeometry& AllottedGeometry, const double InCurrentTime, const float InDeltaTime) OVERRIDE;
-	virtual FReply OnMouseButtonDown(const FGeometry& MyGeometry, const FPointerEvent& MouseEvent) OVERRIDE;
-	virtual FReply OnMouseButtonUp(const FGeometry& MyGeometry, const FPointerEvent& MouseEvent) OVERRIDE;
-	virtual FReply OnMouseMove(const FGeometry& MyGeometry, const FPointerEvent& MouseEvent) OVERRIDE;
-	virtual FReply OnMouseWheel(const FGeometry& MyGeometry, const FPointerEvent& MouseEvent) OVERRIDE;
-	virtual FCursorReply OnCursorQuery( const FGeometry& MyGeometry, const FPointerEvent& CursorEvent ) const OVERRIDE;
-	virtual int32 OnPaint(const FGeometry& AllottedGeometry, const FSlateRect& MyClippingRect, FSlateWindowElementList& OutDrawElements, int32 LayerId, const FWidgetStyle& InWidgetStyle, bool bParentEnabled) const OVERRIDE;
+	virtual void Tick(const FGeometry& AllottedGeometry, const double InCurrentTime, const float InDeltaTime) override;
+	virtual FReply OnMouseButtonDown(const FGeometry& MyGeometry, const FPointerEvent& MouseEvent) override;
+	virtual FReply OnMouseButtonUp(const FGeometry& MyGeometry, const FPointerEvent& MouseEvent) override;
+	virtual FReply OnMouseMove(const FGeometry& MyGeometry, const FPointerEvent& MouseEvent) override;
+	virtual FReply OnMouseWheel(const FGeometry& MyGeometry, const FPointerEvent& MouseEvent) override;
+	virtual FCursorReply OnCursorQuery( const FGeometry& MyGeometry, const FPointerEvent& CursorEvent ) const override;
+	virtual int32 OnPaint(const FGeometry& AllottedGeometry, const FSlateRect& MyClippingRect, FSlateWindowElementList& OutDrawElements, int32 LayerId, const FWidgetStyle& InWidgetStyle, bool bParentEnabled) const override;
 	// End of SWidget interface
 
 	void Construct(const FArguments& InArgs, TSharedRef<class FPaperEditorViewportClient> InViewportClient);
diff --git a/Engine/Plugins/Experimental/Paper2D/Source/Paper2DEditor/Private/SpriteAssetTypeActions.h b/Engine/Plugins/Experimental/Paper2D/Source/Paper2DEditor/Private/SpriteAssetTypeActions.h
--- a/Engine/Plugins/Experimental/Paper2D/Source/Paper2DEditor/Private/SpriteAssetTypeActions.h
+++ b/Engine/Plugins/Experimental/Paper2D/Source/Paper2DEditor/Private/SpriteAssetTypeActions.h
@@ -8,10 +8,10 @@ class FSpriteAssetTypeActions : public FAssetTypeActions_Base
 {
 public:
 	// IAssetTypeActions interface
-	virtual FText GetName() const OVERRIDE;
-	virtual FColor GetTypeColor() const OVERRIDE;
-	virtual UClass* GetSupportedClass() const OVERRIDE;
-	virtual void OpenAssetEditor(const TArray<UObject*>& InObjects, TSharedPtr<class IToolkitHost> EditWithinLevelEditor = TSharedPtr<IToolkitHost>()) OVERRIDE;
-	virtual uint32 GetCategories() OVERRIDE;
+	virtual FText GetName() const override;
+	virtual FColor GetTypeColor() const override;
+	virtual UClass* GetSupportedClass() const override;
+	virtual void OpenAssetEditor(const TArray<UObject*>& InObjects, TSharedPtr<class IToolkitHost> EditWithinLevelEditor = TSharedPtr<IToolkitHost>()) override;
+	virtual uint32 GetCategories() override;
 	// End of IAssetTypeActions interface
 };
\ No newline at end of file
diff --git a/Engine/Plugins/Experimental/Paper2D/Source/Paper2DEditor/Private/SpriteEditor/SSpriteList.h b/Engine/Plugins/Experimental/Paper2D/Source/Paper2DEditor/Private/SpriteEditor/SSpriteList.h
--- a/Engine/Plugins/Experimental/Paper2D/Source/Paper2DEditor/Private/SpriteEditor/SSpriteList.h
+++ b/Engine/Plugins/Experimental/Paper2D/Source/Paper2DEditor/Private/SpriteEditor/SSpriteList.h
@@ -14,7 +14,7 @@ public:
 	void Construct(const FArguments& InArgs, TSharedPtr<class FSpriteEditor> InSpriteEditor);
 
 	// SWidget interface
-	virtual void Tick(const FGeometry& AllottedGeometry, const double InCurrentTime, const float InDeltaTime) OVERRIDE;
+	virtual void Tick(const FGeometry& AllottedGeometry, const double InCurrentTime, const float InDeltaTime) override;
 	// End of SWidget interface
 
 protected:
diff --git a/Engine/Plugins/Experimental/Paper2D/Source/Paper2DEditor/Private/SpriteEditor/SpriteDetailsCustomization.h b/Engine/Plugins/Experimental/Paper2D/Source/Paper2DEditor/Private/SpriteEditor/SpriteDetailsCustomization.h
--- a/Engine/Plugins/Experimental/Paper2D/Source/Paper2DEditor/Private/SpriteEditor/SpriteDetailsCustomization.h
+++ b/Engine/Plugins/Experimental/Paper2D/Source/Paper2DEditor/Private/SpriteEditor/SpriteDetailsCustomization.h
@@ -14,7 +14,7 @@ public:
 	static TSharedRef<IDetailCustomization> MakeInstance();
 
 	// IDetailCustomization interface
-	virtual void CustomizeDetails(IDetailLayoutBuilder& DetailLayout) OVERRIDE;
+	virtual void CustomizeDetails(IDetailLayoutBuilder& DetailLayout) override;
 	// End of IDetailCustomization interface
 
 protected:
diff --git a/Engine/Plugins/Experimental/Paper2D/Source/Paper2DEditor/Private/SpriteEditor/SpriteEditor.cpp b/Engine/Plugins/Experimental/Paper2D/Source/Paper2DEditor/Private/SpriteEditor/SpriteEditor.cpp
--- a/Engine/Plugins/Experimental/Paper2D/Source/Paper2DEditor/Private/SpriteEditor/SpriteEditor.cpp
+++ b/Engine/Plugins/Experimental/Paper2D/Source/Paper2DEditor/Private/SpriteEditor/SpriteEditor.cpp
@@ -49,9 +49,9 @@ public:
 	void Construct(const FArguments& InArgs, TSharedPtr<FSpriteEditor> InSpriteEditor);
 
 	// SEditorViewport interface
-	virtual void BindCommands() OVERRIDE;
-	virtual TSharedRef<FEditorViewportClient> MakeEditorViewportClient() OVERRIDE;
-	virtual TSharedPtr<SWidget> MakeViewportToolbar() OVERRIDE;
+	virtual void BindCommands() override;
+	virtual TSharedRef<FEditorViewportClient> MakeEditorViewportClient() override;
+	virtual TSharedPtr<SWidget> MakeViewportToolbar() override;
 	// End of SEditorViewport interface
 
 	// Invalidate any references to the sprite being edited; it has changed
@@ -220,12 +220,12 @@ public:
 	}
 
 	// SSingleObjectDetailsPanel interface
-	virtual UObject* GetObjectToObserve() const OVERRIDE
+	virtual UObject* GetObjectToObserve() const override
 	{
 		return SpriteEditorPtr.Pin()->GetSpriteBeingEdited();
 	}
 
-	virtual TSharedRef<SWidget> PopulateSlot(TSharedRef<SWidget> PropertyEditorWidget) OVERRIDE
+	virtual TSharedRef<SWidget> PopulateSlot(TSharedRef<SWidget> PropertyEditorWidget) override
 	{
 		return SNew(SVerticalBox)
 			+SVerticalBox::Slot()
diff --git a/Engine/Plugins/Experimental/Paper2D/Source/Paper2DEditor/Private/SpriteEditor/SpriteEditor.h b/Engine/Plugins/Experimental/Paper2D/Source/Paper2DEditor/Private/SpriteEditor/SpriteEditor.h
--- a/Engine/Plugins/Experimental/Paper2D/Source/Paper2DEditor/Private/SpriteEditor/SpriteEditor.h
+++ b/Engine/Plugins/Experimental/Paper2D/Source/Paper2DEditor/Private/SpriteEditor/SpriteEditor.h
@@ -12,20 +12,20 @@ class FSpriteEditor : public FAssetEditorToolkit, public FGCObject
 {
 public:
 	// IToolkit interface
-	virtual void RegisterTabSpawners(const TSharedRef<class FTabManager>& TabManager) OVERRIDE;
-	virtual void UnregisterTabSpawners(const TSharedRef<class FTabManager>& TabManager) OVERRIDE;
+	virtual void RegisterTabSpawners(const TSharedRef<class FTabManager>& TabManager) override;
+	virtual void UnregisterTabSpawners(const TSharedRef<class FTabManager>& TabManager) override;
 	// End of IToolkit interface
 
 	// FAssetEditorToolkit
-	virtual FName GetToolkitFName() const OVERRIDE;
-	virtual FText GetBaseToolkitName() const OVERRIDE;
-	virtual FText GetToolkitName() const OVERRIDE;
-	virtual FLinearColor GetWorldCentricTabColorScale() const OVERRIDE;
-	virtual FString GetWorldCentricTabPrefix() const OVERRIDE;
+	virtual FName GetToolkitFName() const override;
+	virtual FText GetBaseToolkitName() const override;
+	virtual FText GetToolkitName() const override;
+	virtual FLinearColor GetWorldCentricTabColorScale() const override;
+	virtual FString GetWorldCentricTabPrefix() const override;
 	// End of FAssetEditorToolkit
 
 	// FSerializableObject interface
-	virtual void AddReferencedObjects(FReferenceCollector& Collector) OVERRIDE;
+	virtual void AddReferencedObjects(FReferenceCollector& Collector) override;
 	// End of FSerializableObject interface
 
 	// Get the source texture for the current sprite being edited
diff --git a/Engine/Plugins/Experimental/Paper2D/Source/Paper2DEditor/Private/SpriteEditor/SpriteEditorSelections.h b/Engine/Plugins/Experimental/Paper2D/Source/Paper2DEditor/Private/SpriteEditor/SpriteEditorSelections.h
--- a/Engine/Plugins/Experimental/Paper2D/Source/Paper2DEditor/Private/SpriteEditor/SpriteEditorSelections.h
+++ b/Engine/Plugins/Experimental/Paper2D/Source/Paper2DEditor/Private/SpriteEditor/SpriteEditorSelections.h
@@ -94,12 +94,12 @@ public:
 	{
 	}
 
-	virtual uint32 GetTypeHash() const OVERRIDE
+	virtual uint32 GetTypeHash() const override
 	{
 		return VertexIndex;
 	}
 
-	virtual bool Equals(const FSelectedItem& OtherItem) const OVERRIDE
+	virtual bool Equals(const FSelectedItem& OtherItem) const override
 	{
 		if (OtherItem.IsA(FSelectionTypes::SourceRegion))
 		{
@@ -227,22 +227,22 @@ public:
 		return Result;
 	}
 
-	virtual void ApplyDelta(const FVector2D& Delta) OVERRIDE
+	virtual void ApplyDelta(const FVector2D& Delta) override
 	{
 		ApplyDeltaIndexed(Delta, VertexIndex);
 	}
 
-	FVector GetWorldPos() const OVERRIDE
+	FVector GetWorldPos() const override
 	{
 		return GetWorldPosIndexed(VertexIndex);
 	}
 
-	virtual void Delete() OVERRIDE
+	virtual void Delete() override
 	{
 		// Cant delete anything on a bounds object
 	}
 
-		virtual void SplitEdge() OVERRIDE
+		virtual void SplitEdge() override
 	{
 		// Nonsense operation on a vertex, do nothing
 	}
@@ -271,12 +271,12 @@ public:
 	{
 	}
 
-	virtual uint32 GetTypeHash() const OVERRIDE
+	virtual uint32 GetTypeHash() const override
 	{
 		return VertexIndex + (PolygonIndex * 311) + (bRenderData ? 1063 : 0);
 	}
 
-	virtual bool Equals(const FSelectedItem& OtherItem) const OVERRIDE
+	virtual bool Equals(const FSelectedItem& OtherItem) const override
 	{
 		if (OtherItem.IsA(FSelectionTypes::Vertex))
 		{
@@ -336,17 +336,17 @@ public:
 		return Result;
 	}
 
-	virtual void ApplyDelta(const FVector2D& Delta) OVERRIDE
+	virtual void ApplyDelta(const FVector2D& Delta) override
 	{
 		ApplyDeltaIndexed(Delta, VertexIndex);
 	}
 
-	FVector GetWorldPos() const OVERRIDE
+	FVector GetWorldPos() const override
 	{
 		return GetWorldPosIndexed(VertexIndex);
 	}
 
-	virtual void Delete() OVERRIDE
+	virtual void Delete() override
 	{
 		if (UPaperSprite* Sprite = SpritePtr.Get())
 		{
@@ -369,7 +369,7 @@ public:
 		}
 	}
 
-	virtual void SplitEdge() OVERRIDE
+	virtual void SplitEdge() override
 	{
 		// Nonsense operation on a vertex, do nothing
 	}
@@ -390,7 +390,7 @@ public:
 		TypeName = FSelectionTypes::Edge;
 	}
 
-	virtual bool Equals(const FSelectedItem& OtherItem) const OVERRIDE
+	virtual bool Equals(const FSelectedItem& OtherItem) const override
 	{
 		if (OtherItem.IsA(FSelectionTypes::Edge))
 		{
@@ -405,13 +405,13 @@ public:
 		}
 	}
 
-	virtual void ApplyDelta(const FVector2D& Delta) OVERRIDE
+	virtual void ApplyDelta(const FVector2D& Delta) override
 	{
 		ApplyDeltaIndexed(Delta, VertexIndex);
 		ApplyDeltaIndexed(Delta, VertexIndex+1);
 	}
 
-	FVector GetWorldPos() const OVERRIDE
+	FVector GetWorldPos() const override
 	{
 		const FVector Pos1 = GetWorldPosIndexed(VertexIndex);
 		const FVector Pos2 = GetWorldPosIndexed(VertexIndex+1);
@@ -419,12 +419,12 @@ public:
 		return (Pos1 + Pos2) * 0.5f;
 	}
 
-	virtual void Delete() OVERRIDE
+	virtual void Delete() override
 	{
 		//@TODO: Support deleting edges
 	}
 
-	virtual void SplitEdge() OVERRIDE
+	virtual void SplitEdge() override
 	{
 		if (UPaperSprite* Sprite = SpritePtr.Get())
 		{
@@ -464,7 +464,7 @@ public:
 	{
 	}
 
-	virtual bool Equals(const FSelectedItem& OtherItem) const OVERRIDE
+	virtual bool Equals(const FSelectedItem& OtherItem) const override
 	{
 		if (OtherItem.IsA(FSelectionTypes::Socket))
 		{
@@ -480,7 +480,7 @@ public:
 	}
 
 	//@TODO: Currently sockets are in unflipped pivot space,
-	virtual void ApplyDelta(const FVector2D& Delta) OVERRIDE
+	virtual void ApplyDelta(const FVector2D& Delta) override
 	{
 		if (UPaperSprite* Sprite = SpritePtr.Get())
 		{
@@ -492,7 +492,7 @@ public:
 		}
 	}
 
-	FVector GetWorldPos() const OVERRIDE
+	FVector GetWorldPos() const override
 	{
 		if (UPaperSprite* Sprite = SpritePtr.Get())
 		{
@@ -506,12 +506,12 @@ public:
 		return FVector::ZeroVector;
 	}
 
-	virtual void Delete() OVERRIDE
+	virtual void Delete() override
 	{
 		//@TODO: Support deleting edges
 	}
 
-	virtual void SplitEdge() OVERRIDE
+	virtual void SplitEdge() override
 	{
 		// Nonsense operation on a socket, do nothing
 	}
