#ue4 comment update, thanks Thomas S.

diff --git a/Engine/Plugins/Developer/PerforceSourceControl/Source/PerforceSourceControl/Private/PerforceSourceControlOperations.cpp b/Engine/Plugins/Developer/PerforceSourceControl/Source/PerforceSourceControl/Private/PerforceSourceControlOperations.cpp
--- a/Engine/Plugins/Developer/PerforceSourceControl/Source/PerforceSourceControl/Private/PerforceSourceControlOperations.cpp
+++ b/Engine/Plugins/Developer/PerforceSourceControl/Source/PerforceSourceControl/Private/PerforceSourceControlOperations.cpp
@@ -970,7 +970,7 @@ bool FPerforceUpdateStatusWorker::UpdateStates() const
 		const FPerforceSourceControlState& Status = OutStates[StatusIndex];
 		TSharedRef<FPerforceSourceControlState, ESPMode::ThreadSafe> State = PerforceSourceControl.GetProvider().GetStateInternal(Status.LocalFilename);
 		// Update every member except History and Timestamp. History will be updated below from the OutHistory map.
-		// Timestamp has seemingly never been correct.
+		// Timestamp is used to throttle status requests, so update it to current time:
 		auto History = MoveTemp(State->History);
 		*State = Status; 
 		State->History = MoveTemp(History);
