Converted the source control commit comment to use a multi-line editable text control

diff --git a/Engine/Plugins/Developer/PerforceSourceControl/Source/PerforceSourceControl/Private/PerforceConnection.cpp b/Engine/Plugins/Developer/PerforceSourceControl/Source/PerforceSourceControl/Private/PerforceConnection.cpp
--- a/Engine/Plugins/Developer/PerforceSourceControl/Source/PerforceSourceControl/Private/PerforceConnection.cpp
+++ b/Engine/Plugins/Developer/PerforceSourceControl/Source/PerforceSourceControl/Private/PerforceConnection.cpp
@@ -87,11 +87,20 @@ public:
 		OutputDesc += TEXT("\n\n");
 		OutputDesc += TEXT("User:\t");
 		OutputDesc += TO_TCHAR(P4Client.GetUser().Text(), bIsUnicodeServer);
-		OutputDesc += "\n\n";
-		OutputDesc += "Status:\tnew\n\n";
-		OutputDesc += "Description:\n\t";
-		OutputDesc += Description.ToString();
 		OutputDesc += TEXT("\n\n");
+		OutputDesc += TEXT("Status:\tnew\n\n");
+		OutputDesc += TEXT("Description:\n");
+		{
+			TArray<FString> DescLines;
+			Description.ToString().ParseIntoArray(&DescLines, TEXT("\n"), false);
+			for (const FString& DescLine : DescLines)
+			{
+				OutputDesc += TEXT("\t");
+				OutputDesc += DescLine;
+				OutputDesc += TEXT("\n");
+			}
+		}
+		OutputDesc += TEXT("\n");
 		OutputDesc += TEXT("Files:\n\n");
 
 		OutBuffer->Append(FROM_TCHAR(*OutputDesc, bIsUnicodeServer));
diff --git a/Engine/Source/Editor/SourceControlWindows/Private/SSourceControlSubmit.cpp b/Engine/Source/Editor/SourceControlWindows/Private/SSourceControlSubmit.cpp
--- a/Engine/Source/Editor/SourceControlWindows/Private/SSourceControlSubmit.cpp
+++ b/Engine/Source/Editor/SourceControlWindows/Private/SSourceControlSubmit.cpp
@@ -140,12 +140,12 @@ public:
 				.AutoHeight()
 				.Padding(FMargin(5, 0, 5, 5))
 				[
-					SNew( SHorizontalBox )
-					+SHorizontalBox::Slot()
-					.MaxWidth(500.f)
+					SNew(SBox)
+					.WidthOverride(520)
 					[
-						SAssignNew( ChangeListDescriptionTextCtrl, SEditableTextBox )
+						SAssignNew( ChangeListDescriptionTextCtrl, SMultiLineEditableTextBox )
 						.SelectAllTextWhenFocused( true )
+						.AutoWrapText( true )
 					]
 				]
 				+SVerticalBox::Slot()
@@ -447,7 +447,7 @@ private:
 	TWeakPtr<SWindow> ParentFrame;
 
 	/** Internal widgets to save having to get in multiple places*/
-	TSharedPtr<SEditableTextBox> ChangeListDescriptionTextCtrl;
+	TSharedPtr<SMultiLineEditableTextBox> ChangeListDescriptionTextCtrl;
 
 	ESlateCheckBoxState::Type	KeepCheckedOut;
 };
