Fixed Checked Out filter in Content Browser.

diff --git a/Engine/Plugins/Developer/PerforceSourceControl/Source/PerforceSourceControl/Private/PerforceConnection.cpp b/Engine/Plugins/Developer/PerforceSourceControl/Source/PerforceSourceControl/Private/PerforceConnection.cpp
--- a/Engine/Plugins/Developer/PerforceSourceControl/Source/PerforceSourceControl/Private/PerforceConnection.cpp
+++ b/Engine/Plugins/Developer/PerforceSourceControl/Source/PerforceSourceControl/Private/PerforceConnection.cpp
@@ -750,6 +750,9 @@ void FPerforceConnection::EstablishConnection(const FPerforceConnectionInfo& InC
 			{
 				UE_LOG(LogSourceControl, Verbose, TEXT(" ... getting clientroot" ));
 				ClientRoot = Records[0](TEXT("clientRoot"));
+
+				//make sure all slashes point the same way
+				ClientRoot = ClientRoot.Replace(TEXT("\\"), TEXT("/"));
 			}
 		}
 	}
