Prevented Perforce being able to mark files for add that don't currently exist on disk

UE-17117 - The Perforce SCC allows you to add files that don't exist

This also allows us to remove the ShouldForceBinaryAdd function, as that was to fix bad behavior when Perforce was addding an asset file that didn't yet exist.

diff --git a/Engine/Plugins/Developer/PerforceSourceControl/Source/PerforceSourceControl/Private/PerforceSourceControlOperations.cpp b/Engine/Plugins/Developer/PerforceSourceControl/Source/PerforceSourceControl/Private/PerforceSourceControlOperations.cpp
--- a/Engine/Plugins/Developer/PerforceSourceControl/Source/PerforceSourceControl/Private/PerforceSourceControlOperations.cpp
+++ b/Engine/Plugins/Developer/PerforceSourceControl/Source/PerforceSourceControl/Private/PerforceSourceControlOperations.cpp
@@ -361,53 +361,33 @@ FName FPerforceMarkForAddWorker::GetName() const
 	return "MarkForAdd";
 }
 
-static bool ShouldForceBinaryAdd(const TArray<FString>& FilesToAdd)
+bool FPerforceMarkForAddWorker::Execute(FPerforceSourceControlCommand& InCommand)
 {
-	bool bForceBinary = false;
-
-	if ( FilesToAdd.Num() > 0 )
+	// Perforce will allow you to mark files for add that don't currently exist on disk
+	// This goes against the workflow of our other SCC providers (such as SVN and Git), 
+	// so we manually check that the files exist before allowing this command to continue
+	// This keeps the behavior consistent between SCC providers
+	bool bHasMissingFiles = false;
+	for(const FString& FileToAdd : InCommand.Files)
 	{
-		// Force a binary type when adding if all files are .uasset files and at least one does not exist on disk.
-		// If a .uasset file is marked for add before it exists on disk, it will be added to perforce in text format.
-		// This will corrupt the file and cause a crash at startup when scanning the file with the asset registry.
-		// If the user is adding one or more uasset files and any are missing it will force binary type for all of them.
-		bool bAllFilesUAsset = true;
-		bool bAtLeastOneFileDoesNotExist = false;
-		for ( int32 ParamIdx = 0; ParamIdx < FilesToAdd.Num(); ++ParamIdx )
+		if(!IFileManager::Get().FileExists(*FileToAdd))
 		{
-			const FString& Param = FilesToAdd[ParamIdx];
-			if ( FPaths::GetExtension(Param, true).ToLower() == FPackageName::GetAssetPackageExtension().ToLower() )
-			{
-				if ( !ensureMsgf(FPaths::FileExists(Param), TEXT("Attempting to add a uasset file '%s' to Perforce that does not yet exist."), *Param) )
-				{
-					// This is a uasset file that does not yet exist on disk. This is invalid, find out why this is happening.
-					bAtLeastOneFileDoesNotExist = true;
-				}
-			}
-			else
-			{
-				bAllFilesUAsset = false;
-			}
+			bHasMissingFiles = true;
+			InCommand.ErrorMessages.Add(FText::Format(LOCTEXT("Error_FailedToMarkFileForAdd_FileMissing", "Failed mark the file '%s' for add. The file doesn't exist on disk."), FText::FromString(FileToAdd)));
 		}
-
-		bForceBinary = bAllFilesUAsset && bAtLeastOneFileDoesNotExist;
+	}
+	if(bHasMissingFiles)
+	{
+		InCommand.bCommandSuccessful = false;
+		return false;
 	}
 
-	return bForceBinary;
-}
-
-bool FPerforceMarkForAddWorker::Execute(FPerforceSourceControlCommand& InCommand)
-{
 	FScopedPerforceConnection ScopedConnection(InCommand);
 	if(ScopedConnection.IsValid())
 	{
 		FPerforceConnection& Connection = ScopedConnection.GetConnection();
 		TArray<FString> Parameters;
 		FP4RecordSet Records;
-		if ( ShouldForceBinaryAdd(InCommand.Files) )
-		{
-			Parameters.Add(TEXT("-t binary+l"));
-		}
 		Parameters.Append(InCommand.Files);
 		InCommand.bCommandSuccessful = Connection.RunCommand(TEXT("add"), Parameters, Records, InCommand.ErrorMessages, FOnIsCancelled::CreateRaw(&InCommand, &FPerforceSourceControlCommand::IsCanceled), InCommand.bConnectionDropped);
 		ParseRecordSetForState(Records, OutResults);
