Perforce library support switch.
Allows PerforceSourceControl plugin to be built without the P4API libraries (will not function of course).
Perforce plugin build switched off for VS2015 until we can get libs in place.

diff --git a/Engine/Plugins/Developer/PerforceSourceControl/Source/PerforceSourceControl/PerforceSourceControl.Build.cs b/Engine/Plugins/Developer/PerforceSourceControl/Source/PerforceSourceControl/PerforceSourceControl.Build.cs
--- a/Engine/Plugins/Developer/PerforceSourceControl/Source/PerforceSourceControl/PerforceSourceControl.Build.cs
+++ b/Engine/Plugins/Developer/PerforceSourceControl/Source/PerforceSourceControl/PerforceSourceControl.Build.cs
@@ -18,7 +18,15 @@ public class PerforceSourceControl : ModuleRules
 			}
 		);
 
+		//@todo: Perforce API libraries are not yet available for VS2015
+		//if ((Target.Platform == UnrealTargetPlatform.Win64 || Target.Platform == UnrealTargetPlatform.Win32) && WindowsPlatform.Compiler == WindowsCompiler.VisualStudio2015)
+		//{
+		//	Definitions.Add("USE_P4_API=0");
+		//	return;
+		//}
+
 		AddThirdPartyPrivateStaticDependencies(Target, "Perforce");
+		Definitions.Add("USE_P4_API=1");
 
 		if (Target.Platform == UnrealTargetPlatform.Win64 || Target.Platform == UnrealTargetPlatform.Mac)
 		{
diff --git a/Engine/Plugins/Developer/PerforceSourceControl/Source/PerforceSourceControl/Private/PerforceConnection.cpp b/Engine/Plugins/Developer/PerforceSourceControl/Source/PerforceSourceControl/Private/PerforceConnection.cpp
--- a/Engine/Plugins/Developer/PerforceSourceControl/Source/PerforceSourceControl/Private/PerforceConnection.cpp
+++ b/Engine/Plugins/Developer/PerforceSourceControl/Source/PerforceSourceControl/Private/PerforceConnection.cpp
@@ -29,6 +29,7 @@ public:
 	// Called by P4API when the results from running a command are ready
 	void OutputStat(StrDict* VarList)
 	{
+#if USE_P4_API
 		FP4Record Record;
 		StrRef Var, Value;
 
@@ -39,13 +40,16 @@ public:
 		}
 
 		Records.Add(Record);
+#endif
 	}
 
 	void HandleError(Error *InError)
 	{
+#if USE_P4_API
 		StrBuf ErrorMessage;
 		InError->Fmt(&ErrorMessage);
 		OutErrorMessages.Add(FText::FromString(FString(TO_TCHAR(ErrorMessage.Text(), bIsUnicodeServer))));
+#endif
 	}
 
 	bool bIsUnicodeServer;
@@ -80,6 +84,7 @@ public:
 	// Called by P4API on "change -i" command. OutBuffer is filled with changelist specification text
 	void InputData(StrBuf* OutBuffer, Error* OutError)
 	{
+#if USE_P4_API
 		FString OutputDesc;
 		OutputDesc += TEXT("Change:\tnew\n\n");
 		OutputDesc += TEXT("Client:\t");
@@ -104,6 +109,7 @@ public:
 		OutputDesc += TEXT("Files:\n\n");
 
 		OutBuffer->Append(FROM_TCHAR(*OutputDesc, bIsUnicodeServer));
+#endif
 	}
 
 	FText Description;
@@ -125,7 +131,9 @@ public:
 
 	void Prompt( const StrPtr& InMessage, StrBuf& OutPrompt, int NoEcho, Error* InError )
 	{
+#if USE_P4_API
 		OutPrompt.Set(FROM_TCHAR(*Password, bIsUnicodeServer));
+#endif
 	}
 
 	/** Password to use when logging in */
@@ -141,6 +149,7 @@ public:
 	{
 	}
 
+#if USE_P4_API
 	virtual int IsAlive() override
 	{
 		if(IsCancelled.IsBound() && IsCancelled.Execute())
@@ -150,6 +159,7 @@ public:
 
 		return 1;
 	}
+#endif
 
 	FOnIsCancelled IsCancelled;
 };
@@ -177,9 +187,11 @@ static bool TestConnection(ClientApi& P4Client, const FString& ClientSpecName, b
 	}
 
 	const char *ArgV[] = { "-o", (char*)ClientSpecUTF8Name };
+#if USE_P4_API
 	P4Client.SetArgv(2, const_cast<char*const*>(ArgV));
 
 	P4Client.Run("client", &User);
+#endif
 
 	// clean up args
 	delete [] ClientSpecUTF8Name;
@@ -197,6 +209,7 @@ static bool TestConnection(ClientApi& P4Client, const FString& ClientSpecName, b
 
 static bool CheckUnicodeStatus(ClientApi& P4Client, bool& bIsUnicodeServer, TArray<FText>& OutErrorMessages)
 {
+#if USE_P4_API
 	FP4RecordSet Records;
 	FP4ClientUser User(Records, false, OutErrorMessages);
 
@@ -210,6 +223,7 @@ static bool CheckUnicodeStatus(ClientApi& P4Client, bool& bIsUnicodeServer, TArr
 	{
 		bIsUnicodeServer = false;
 	}
+#endif
 
 	return OutErrorMessages.Num() == 0;
 }
@@ -268,6 +282,7 @@ bool FPerforceConnection::AutoDetectWorkspace(const FPerforceConnectionInfo& InC
 bool FPerforceConnection::Login(const FPerforceConnectionInfo& InConnectionInfo)
 {
 	TArray<FText> ErrorMessages;
+#if USE_P4_API
 	FP4RecordSet Records;
 	FP4LoginClientUser User(InConnectionInfo.Password, Records, false, ErrorMessages);
 
@@ -283,6 +298,7 @@ bool FPerforceConnection::Login(const FPerforceConnectionInfo& InConnectionInfo)
 			UE_LOG(LogSourceControl, Error, TEXT("%s"), *ErrorMessage.ToString());
 		}
 	}
+#endif
 
 	return ErrorMessages.Num() == 0;
 }
@@ -291,6 +307,7 @@ bool FPerforceConnection::EnsureValidConnection(FString& InOutServerName, FStrin
 {
 	bool bIsUnicodeServer = false;
 	bool bConnectionOK = false;
+#if USE_P4_API
 	FMessageLog SourceControlLog("SourceControl");
 
 	FString NewServerName = InOutServerName;
@@ -447,6 +464,7 @@ bool FPerforceConnection::EnsureValidConnection(FString& InOutServerName, FStrin
 		InOutUserName = NewUserName;
 		InOutWorkspaceName = NewClientSpecName;
 	}
+#endif
 
 	return bConnectionOK;
 }
@@ -541,11 +559,16 @@ bool FPerforceConnection::GetWorkspaceList(const FPerforceConnectionInfo& InConn
 
 bool FPerforceConnection::IsValidConnection()
 {
+#if USE_P4_API
 	return bEstablishedConnection && !P4Client.Dropped();
+#else
+	return false;
+#endif
 }
 
 void FPerforceConnection::Disconnect()
 {
+#if USE_P4_API
 	Error P4Error;
 
 	P4Client.Final(&P4Error);
@@ -557,10 +580,12 @@ void FPerforceConnection::Disconnect()
 		UE_LOG(LogSourceControl, Error, TEXT("P4ERROR: Failed to disconnect from Server."));
 		UE_LOG(LogSourceControl, Error, TEXT("%s"), TO_TCHAR(ErrorMessage.Text(), bIsUnicode));
 	}
+#endif
 }
 
 bool FPerforceConnection::RunCommand(const FString& InCommand, const TArray<FString>& InParameters, FP4RecordSet& OutRecordSet, TArray<FText>& OutErrorMessage, FOnIsCancelled InIsCancelled, bool& OutConnectionDropped, const bool bInStandardDebugOutput, const bool bInAllowRetry)
 {
+#if USE_P4_API
 	if (!bEstablishedConnection)
 	{
 		return false;
@@ -625,12 +650,14 @@ bool FPerforceConnection::RunCommand(const FString& InCommand, const TArray<FStr
 	{
 		UE_LOG( LogSourceControl, VeryVerbose, TEXT("P4 execution time: %0.4f seconds. Command: %s"), FPlatformTime::Seconds() - SCCStartTime, *FullCommand );
 	}
+#endif
 
 	return OutRecordSet.Num() > 0;
 }
 
 int32 FPerforceConnection::CreatePendingChangelist(const FText &Description, FOnIsCancelled InIsCancelled, TArray<FText>& OutErrorMessages)
 {
+#if USE_P4_API
 	TArray<FString> Params;
 	FP4RecordSet Records;
 
@@ -646,10 +673,14 @@ int32 FPerforceConnection::CreatePendingChangelist(const FText &Description, FOn
 	P4Client.SetBreak(NULL);
 
 	return User.ChangelistNumber;
+#else
+	return 0;
+#endif
 }
 
 void FPerforceConnection::EstablishConnection(const FPerforceConnectionInfo& InConnectionInfo)
 {
+#if USE_P4_API
 	// Verify Input. ServerName and UserName are required
 	if ( InConnectionInfo.Port.IsEmpty() || InConnectionInfo.UserName.IsEmpty() )
 	{
@@ -756,6 +787,7 @@ void FPerforceConnection::EstablishConnection(const FPerforceConnectionInfo& InC
 			}
 		}
 	}
+#endif
 }
 
 FScopedPerforceConnection::FScopedPerforceConnection( const class FPerforceSourceControlCommand& InCommand )
diff --git a/Engine/Plugins/Developer/PerforceSourceControl/Source/PerforceSourceControl/Private/PerforceConnection.h b/Engine/Plugins/Developer/PerforceSourceControl/Source/PerforceSourceControl/Private/PerforceConnection.h
--- a/Engine/Plugins/Developer/PerforceSourceControl/Source/PerforceSourceControl/Private/PerforceConnection.h
+++ b/Engine/Plugins/Developer/PerforceSourceControl/Source/PerforceSourceControl/Private/PerforceConnection.h
@@ -94,8 +94,10 @@ public:
 	void EstablishConnection(const FPerforceConnectionInfo& InConnectionInfo);
 
 public:
+#if USE_P4_API
 	/** Perforce API client object */
 	ClientApi P4Client;
+#endif
 
 	/** The current root for the client workspace */
 	FString ClientRoot;
diff --git a/Engine/Plugins/Developer/PerforceSourceControl/Source/PerforceSourceControl/Private/PerforceSourceControlOperations.cpp b/Engine/Plugins/Developer/PerforceSourceControl/Source/PerforceSourceControl/Private/PerforceSourceControlOperations.cpp
--- a/Engine/Plugins/Developer/PerforceSourceControl/Source/PerforceSourceControl/Private/PerforceSourceControlOperations.cpp
+++ b/Engine/Plugins/Developer/PerforceSourceControl/Source/PerforceSourceControl/Private/PerforceSourceControlOperations.cpp
@@ -860,6 +860,7 @@ FName FPerforceUpdateStatusWorker::GetName() const
 
 bool FPerforceUpdateStatusWorker::Execute(FPerforceSourceControlCommand& InCommand)
 {
+#if USE_P4_API
 	FScopedPerforceConnection ScopedConnection(InCommand);
 	if(ScopedConnection.IsValid())
 	{
@@ -944,6 +945,7 @@ bool FPerforceUpdateStatusWorker::Execute(FPerforceSourceControlCommand& InComma
 			ParseDiffResults(Records, OutModifiedFiles);
 		}
 	}
+#endif
 
 	return InCommand.bCommandSuccessful;
 }
diff --git a/Engine/Plugins/Developer/PerforceSourceControl/Source/PerforceSourceControl/Private/PerforceSourceControlPrivatePCH.h b/Engine/Plugins/Developer/PerforceSourceControl/Source/PerforceSourceControl/Private/PerforceSourceControlPrivatePCH.h
--- a/Engine/Plugins/Developer/PerforceSourceControl/Source/PerforceSourceControl/Private/PerforceSourceControlPrivatePCH.h
+++ b/Engine/Plugins/Developer/PerforceSourceControl/Source/PerforceSourceControl/Private/PerforceSourceControlPrivatePCH.h
@@ -6,9 +6,30 @@
 #include "SlateBasics.h"
 #include "EditorStyle.h"
 
+//@todo: perforce doesn't currently support VS2015
+#if USE_P4_API
 PRAGMA_DISABLE_SHADOW_VARIABLE_WARNINGS
 #include <p4/clientapi.h>
 #include <p4/i18napi.h>
 PRAGMA_ENABLE_SHADOW_VARIABLE_WARNINGS
+#endif
+
+//@todo: perforce doesn't currently support VS2015
+#if !USE_P4_API
+class ClientApi
+{};
+class ClientUser
+{};
+class StrBuf
+{};
+class StrDict
+{};
+class Error
+{};
+class StrPtr
+{};
+class KeepAlive
+{};
+#endif
 
 #include "ISourceControlModule.h"
\ No newline at end of file
diff --git a/Engine/Plugins/Developer/PerforceSourceControl/Source/PerforceSourceControl/Private/SPerforceSourceControlSettings.cpp b/Engine/Plugins/Developer/PerforceSourceControl/Source/PerforceSourceControl/Private/SPerforceSourceControlSettings.cpp
--- a/Engine/Plugins/Developer/PerforceSourceControl/Source/PerforceSourceControl/Private/SPerforceSourceControlSettings.cpp
+++ b/Engine/Plugins/Developer/PerforceSourceControl/Source/PerforceSourceControl/Private/SPerforceSourceControlSettings.cpp
@@ -21,6 +21,7 @@ void SPerforceSourceControlSettings::Construct(const FArguments& InArgs)
 
 	if (PortName.IsEmpty() && UserName.IsEmpty())
 	{
+#if USE_P4_API
 		ClientApi TestP4;
 		Error P4Error;
 		TestP4.Init(&P4Error);
@@ -31,6 +32,7 @@ void SPerforceSourceControlSettings::Construct(const FArguments& InArgs)
 		PerforceSourceControl.AccessSettings().SetPort(PortName);
 		PerforceSourceControl.AccessSettings().SetUserName(UserName);
 		PerforceSourceControl.SaveSettings();
+#endif
 	}
 
 	FSlateFontInfo Font = FEditorStyle::GetFontStyle(TEXT("SourceControl.LoginWindow.Font"));
