Enable 2015 p4 lib disable

diff --git a/Engine/Plugins/Developer/PerforceSourceControl/Source/PerforceSourceControl/PerforceSourceControl.Build.cs b/Engine/Plugins/Developer/PerforceSourceControl/Source/PerforceSourceControl/PerforceSourceControl.Build.cs
--- a/Engine/Plugins/Developer/PerforceSourceControl/Source/PerforceSourceControl/PerforceSourceControl.Build.cs
+++ b/Engine/Plugins/Developer/PerforceSourceControl/Source/PerforceSourceControl/PerforceSourceControl.Build.cs
@@ -19,11 +19,11 @@ public class PerforceSourceControl : ModuleRules
 		);
 
 		//@todo: Perforce API libraries are not yet available for VS2015
-		//if ((Target.Platform == UnrealTargetPlatform.Win64 || Target.Platform == UnrealTargetPlatform.Win32) && WindowsPlatform.Compiler == WindowsCompiler.VisualStudio2015)
-		//{
-		//	Definitions.Add("USE_P4_API=0");
-		//	return;
-		//}
+		if ((Target.Platform == UnrealTargetPlatform.Win64 || Target.Platform == UnrealTargetPlatform.Win32) && WindowsPlatform.Compiler == WindowsCompiler.VisualStudio2015)
+		{
+			Definitions.Add("USE_P4_API=0");
+			return;
+		}
 
 		AddThirdPartyPrivateStaticDependencies(Target, "Perforce");
 		Definitions.Add("USE_P4_API=1");
