Copying //UE4/Dev-Enterprise to //UE4/Dev-Main (Source: //UE4/Dev-Enterprise @ 3972172)

============================
  MAJOR FEATURES & CHANGES
============================

Change 3821754 by [REDACTED_PERSON]

	[Python] Class property conversion now goes through NativizeClass/PythonizeClass

	This allows it to coerce from Python wrapped object types

Change 3833107 by [REDACTED_PERSON]

	Added functions to fill an existing DataTable from an existing CSV/JSON file.

Change 3835044 by [REDACTED_PERSON]

	Exposure for asset_import_data (editor property) and it's functions: extract_filenames and get_first_filename.

Change 3835466 by [REDACTED_PERSON]

	Hide function from Python that need special compile command to be executed by the VM.

Change 3839237 by [REDACTED_PERSON]

	Added a way to inspect the full chain of properties that are currently being serialized by an archive

	You used to only have access to the leaf-most property, and while you could use its outer chain to inspect other properties within the same object/struct, you couldn't always get the full chain (eg, if you had an object containing a struct).

Change 3839974 by [REDACTED_PERSON]

	Make sure that SerializedProperty is copied correctly, as SetSerializedPropertyChain may set it to something else

Change 3842311 by [REDACTED_PERSON]

	Fixing potential null level assert

Change 3842313 by [REDACTED_PERSON]

	Updated settings editor to handle external properties

Change 3842316 by [REDACTED_PERSON]

	Allowing a console command to be given to GEditor/GEngine even if there's a player

	CL# 1848982 said it was to prevent multiple execution of stat commands, however that no longer seems to be an issue.

Change 3842867 by [REDACTED_PERSON]

	Added a way to generate diffs from editor transactions

	The notifications from these diffs are send to UObject::PostTransacted and FCoreUObjectDelegates::OnObjectTransacted.

	These notifications are typically generated when a transaction is "finalized", but can also be generated from "snapshots" (eg, to trap nodes being dragged in the world). They're also generated from normal undo/redo events.

Change 3844428 by [REDACTED_PERSON]

	Move the SetMaterial code from the StaticMeshEditor to StaticMesh to be reusable by script.

Change 3845966 by [REDACTED_PERSON]

	Added support for minimal game RPC worlds

	These can be created in the editor and engine and exist to allow RPC communication via Unreal Networking in a way that is sandboxed from any other worlds that may be loaded (like the main game world)

Change 3848844 by [REDACTED_PERSON]

	Expose EComponentMobility to blueprint.

Change 3854616 by [REDACTED_PERSON]

	Add Custom way time step the engine loop. Will be used by the Synchronization of media for enterprise.

Change 3856650 by [REDACTED_PERSON]

	Fixed a bug where transaction finalization could miss changes since the last snapshot

Change 3864951 by [REDACTED_PERSON]

	Fix ghost asset in Content Browser when an asset is added and renamed before the RecentlyAddedAssets list had a chance to be processed.

Change 3867158 by JeanMichel.Dignard

	UBT
	- Added the ability for dll programs to export symbols.

Change 3872342 by [REDACTED_PERSON]

	Merging static analysis fixes from 4.19

Change 3879305 by [REDACTED_PERSON]

	Improved the processing of py files from exec commands

	The old logic used to just test if the entire command was a .py file. The new logic extracts out the first token and sees if that's a .py file, and if it is, treats the remaining data as extra arguments.

Change 3879306 by [REDACTED_PERSON]

	Added a minimal commandlet for invoking Python scripts

Change 3881631 by [REDACTED_PERSON]

	Added basic RTTI to Python meta-data types

Change 3885384 by [REDACTED_PERSON]

	[Python] Prevent glue code using reserved names

Change 3888957 by [REDACTED_PERSON]

	In MediaPlayer, only create a PlayerFacade & Playlist when it's not a ClassDefaultObject.
	The MediaPlayerFacade is a MediaTickable. That trigger the tick thread to be awake even if there is no Media playing.

Change 3888961 by [REDACTED_PERSON]

	Fix FInterval::IsValid return type.

Change 3888980 by [REDACTED_PERSON]

	Modification to Media and MediaAsset to support MediaSmith. The TInterval<int64> will be changed into TTinterval<FTimespan> UEENT-947. MediaSampleQueue's critical section will be change into an atomic operation UEENT-948.

Change 3889165 by [REDACTED_PERSON]

	Fix build. Missing include for Timespan. Introduce with CL 3888980.

Change 3889261 by [REDACTED_PERSON]

	[Python] Fixing some more name conflicts in generated code

Change 3889504 by [REDACTED_PERSON]

	Add option to change PreferredPixelFormat

Change 3891193 by [REDACTED_PERSON]

	Fix build. Missing include for Interval. Introduce with CL 3888980.

Change 3897108 by [REDACTED_PERSON]

	TTinterval use it own traits. Create a Interval traits for Timespan.

Change 3899669 by [REDACTED_PERSON]

	Fixed Functions sometimes being exposed to Python as if they were Structs

Change 3900692 by [REDACTED_PERSON]

	Removed some boilerplate associated with wrapping a basic type to Python

	You can now derive from TPyWrapperBasic to wrap a type that is simply a value copied into Python (see FPyWrapperName and FPyWrapperText for an example)

Change 3901066 by conan.reis

	UE4 editor script bindings (Cobra) and helper functions for version control
	- exposed SourceControl class with common source control methods and associated SourceControlState structure
	- commands have smart file strings that can convert from any of fully qualified path, relative path, long package name, asset path or export text path (often stored on clipboard)
	- commands store any errors in a shared error text object which is optionally printed to the error log
	- renamed some calls across the UE4 codebase to USourceControlHelpers::CheckOutOrAddFile() from USourceControlHelpers::CheckOutFile()
	- included Python test script for source control commands including that auto-creates test files as needed and passes various types of files to test as command line arguments. Any unexpected results displays error messages.

Change 3901388 by [REDACTED_PERSON]

	Minimal Slate hooks for Python

Change 3901456 by [REDACTED_PERSON]

	Added missing file

Change 3901549 by [REDACTED_PERSON]

	Removing some more Windows defines that were causing build issues

Change 3904518 by conan.reis

	Source Control
	- ensured that "check if modified" flag is set whenver getting source control state in USourceControlHelpers::QueryFileState() which was needed when using Perforce source control provider

Change 3905612 by [REDACTED_PERSON]

	Optimize RemoveDuplicates somewhat using a TSet

Change 3912626 by [REDACTED_PERSON]

	Fixed ShouldExcludeDerivedClasses option not working

	RecursiveClassesExclusionSet requires a base ClassNames entry to operate on when filtering.

Change 3917739 by [REDACTED_PERSON]

	Output Log suggestions list is now clamped to the work area width of the monitor that hosts the widget

Change 3917744 by [REDACTED_PERSON]

	Changed generated code to reference the UProperty and UFunction directly, rather than constantly look them up by name

	Names were originally used because UHT couldn't access the objects when it registered the glue code, but now that we generate at runtime via reflection, we already have the relevant objects available, and caching them the glue structs helps performance at both generation time and runtime.

Change 3918832 by [REDACTED_PERSON]

	Removed field iteration from Python function calls

	We now cache the input and output parameters for all function calls (methods, get/set, and delegates) and use this rather than iterate the struct fields.

Change 3920648 by [REDACTED_PERSON]

	Remove the bottom right part of the windows border of the grabbed frame when in the editor. Tested in the standalone, windowed and full screen.
	Add option to request a FlushOnDraw on the viewport. Flushing in SDI output flow decreases the performance by ~10ms. SDI output is synchronized and the engine tick follows that synchronization.

Change 3921396 by [REDACTED_PERSON]

	Split up the generated type data to correspond to the type being wrapped

	A lot of types can just use the minimal set, but classes and structs have some extra data.

Change 3921619 by conan.reis

	- add delegate to FSourceControlWindows::ChoosePackagesToCheckin() that gives info for result, result description, files added, files checked in and flag indicating whether files were checked out again.
	- also added result info to FSourceControlWindows::PromptForCheckin()

Change 3921624 by conan.reis

	Removed Source Control common files from pre compiled header
	- main changes are in UnrealEdPrivatePCH.h, UnrealEdSharedPCH.h, SouceControlWindows.h and the added SouceControlWindows.cpp
	- remaining files have includes changed to accomodate

Change 3921958 by conan.reis

	Fix attempt for incude file dependency needed by some build configurations (likely PCH disabled) caused by CL3921619

Change 3922740 by conan.reis

	Included SourceControlOperations.h and SourceControlHelpers.h back in ISourceControlProvider.h though it does not need them since other files that were including ISourceControlProvider.h have come to expect their inclusion.
	They were previously removed in CL3921624

Change 3923375 by [REDACTED_PERSON]

	Added optimized FString <-> icu::UnicodeString conversion for platforms using UTF-16 native strings

Change 3926547 by [REDACTED_PERSON]

	Added support for struct method "hoisting"

	This allows you to tag a helper function that takes a struct as its first argument with the ScriptMethod meta-data (optionally providing a new name) to "hoist" that helper function to be a method of the struct it operates on when wrapped for Python.

Change 3927050 by conan.reis

	Source control - ensured that  ISourceControlProvider::Execute(FConnect, EConcurrency, FSourceControlOperationComplete) delegate is called on initial connection even if it fails immediately. Modified Perforce, Git and Subversion source control providers

Change 3929268 by conan.reis

	- fixed case in Perforce source control code where the server available flag was set even when the server was not successfully connected
	- removed Perforce error message about file folders outside of the workspace client mappings
	- clarified comments for ISourceControlProvider::IsEnabled() and ISourceControlModule::IsEnabled()

Change 3931024 by [REDACTED_PERSON]

	Expose FBX and Texture import to python

Change 3931273 by [REDACTED_PERSON]

	Hide re-import slate notification pop-up during python automated asset import

Change 3931368 by [REDACTED_PERSON]

	Stopped bools coercing to numeric types in Python nativization

Change 3931374 by [REDACTED_PERSON]

	Added support for struct math operator "hoisting"

	This allows you to tag a helper function that takes a struct as its first argument with the ScriptMathOp meta-data (providing a potentially semi-colon separated list of operators to map to) to "hoist" that helper function to be a math operator of the struct it operates on when wrapped for Python.

Change 3932586 by [REDACTED_PERSON]

	Removed file read into unused memory buffer during Fbx import

Change 3934308 by [REDACTED_PERSON]

	Added a public interface for the Python plugin

	Very basic, just lets you query if Python is compiled in, and lets you execute Python commands like you would via the Output Log.

Change 3935088 by conan.reis

	- Added info/warning and error message storage to all the source control operation structures so additional information can be made available.
	- Added ISourceControlOperation::GetResultInfo() which returns the modified control structures (mentioned above) with appended info/warning messages and error messages and implemented its use in all source control operations in Perforce, Git and Subversion.

Change 3936668 by [REDACTED_PERSON]

	Avoid re-allocation of memory buffer holding file bytes during asset import

Change 3940596 by [REDACTED_PERSON]

	Overlapping vertex check was O (N^2)
	100k vertex mesh took ~15 seconds to perform overlap step now takes 0.023 seconds

Change 3942629 by [REDACTED_PERSON]

	Fixes a memory leak of FbxScene and reduces wait time during import.

Change 3942884 by [REDACTED_PERSON]

	Python asset import can now customize destination asset name

Change 3946278 by [REDACTED_PERSON]

	Added stricter conversion for math operator arguments

	PyConversion now returns FPyConversionResult rather than bool, which will tell you not only whether a conversion succeeded or failed, but also whether type coercion was applied during the conversion.

	This allows the operator stack evaluation to run a first pass looking for an exact argument match, before falling back to a coerced match if available. This allows operators to apply correctly to coerced types (eg, int vs float overloads).

Change 3948455 by [REDACTED_PERSON]

	Added generic Tick function to FPythonScriptPlugin

	This can also handle init logic for after the engine is fully initialized

Change 3948888 by [REDACTED_PERSON]

	Added settings for the Python plugin

	You can now define start-up scripts to execute once the engine is initialized, additional system paths for Python, and whether you want to enable developer mode (which will enable things like deprecation warnings).

Change 3948982 by [REDACTED_PERSON]

	Fixed Python 3 build error caused by CObject being removed in Python 3.2

Change 3949614 by [REDACTED_PERSON]

	Create a camera cut track from the camera switcher camera index animation curve when importing a fbx in sequencer

Change 3950829 by [REDACTED_PERSON]

	Update error message to be more specific when ENGINE_API keyword is found before 'static' keyword for a UFUNCTION

Change 3953452 by [REDACTED_PERSON]

	Fixed some dependencies

Change 3953645 by [REDACTED_PERSON]

	Fixed Python parameter packing treating bool output paramers as potential return values

	void GetState(bool& OutState) would have previously triggered the code for packing output for a function that returns a bool.

Change 3953850 by [REDACTED_PERSON]

	Fixed doc string generation for a function with multiple output paramters and no return value

Change 3954279 by [REDACTED_PERSON]

	Initial support for exposing deprecated properties and functions to Python

	This handles properties and functions that are directly deprecated. We still need to handle the cases where they're renamed and a redirector is left.

Change 3954922 by [REDACTED_PERSON]

	Expose UnloadPackages to python

Change 3955209 by [REDACTED_PERSON]

	Initial support for exposing deprecated classes to Python

Change 3955248 by [REDACTED_PERSON]

	Added a way to load Unreal modules via Python

	unreal.load_module("modulename")

Change 3955561 by [REDACTED_PERSON]

	Expose asset export to python

Change 3956068 by [REDACTED_PERSON]

	Linux compile fix.

Change 3960449 by [REDACTED_PERSON]

	Fix automated test using bCombineMeshes

Change 3960495 by [REDACTED_PERSON]

	Add a temporary menu to show the MetaData of an asset.
	The menu will need to be updated to have a look and feel of the Detail View and support edition at one point.

Change 3961599 by [REDACTED_PERSON]

	Reduced peak memory during import of meshes related to duplicate vertex tracking

Change 3962104 by [REDACTED_PERSON]

	Disable import mesh overlapping corners memory optimization to because it can change uv generation

Change 3962507 by [REDACTED_PERSON]

	Fix uv generation

Change 3965285 by [REDACTED_PERSON]

	Add support for FBX export as ASCII

Change 3965287 by [REDACTED_PERSON]

	Forgotten file, fbx export as ascii

Change 3966772 by [REDACTED_PERSON]

	Fix MaterialExpressionFunctions for ExternalTexture support

Change 3967014 by [REDACTED_PERSON]

	Added a way to get the CDO in Python

	Wrapped objects now have a get_default_object class method

Change 3967151 by [REDACTED_PERSON]

	Added stats to track Python generation time

Change 3968006 by [REDACTED_PERSON]

	Media Samples

	- Removed Locks and Min/Max SampleTime from queues
	- Added methods to fetch NextSampleTime and SampleCount in queues
	- Added MediaSource base class for players that want to be time synchronized

Change 3969119 by [REDACTED_PERSON]

	Add delay functionnality to MediaPlayer to delay the frame by some time. It will allow more than one player to be start at the same time, played at the same frame but offset in relation to each other.

diff --git a/Engine/Config/BaseEditorPerProjectUserSettings.ini b/Engine/Config/BaseEditorPerProjectUserSettings.ini
--- a/Engine/Config/BaseEditorPerProjectUserSettings.ini
+++ b/Engine/Config/BaseEditorPerProjectUserSettings.ini
@@ -62,6 +62,7 @@ bDisplayEngineVersionInBadge=False
 
 [/Script/UnrealEd.FbxExportOption]
 FbxExportCompatibility=FBX_2013
+bASCII=false
 bForceFrontXAxis=false
 LevelOfDetail=True
 Collision=True
diff --git a/Engine/Plugins/Developer/GitSourceControl/Source/GitSourceControl/Private/GitSourceControlCommand.cpp b/Engine/Plugins/Developer/GitSourceControl/Source/GitSourceControl/Private/GitSourceControlCommand.cpp
--- a/Engine/Plugins/Developer/GitSourceControl/Source/GitSourceControl/Private/GitSourceControlCommand.cpp
+++ b/Engine/Plugins/Developer/GitSourceControl/Source/GitSourceControl/Private/GitSourceControlCommand.cpp
@@ -38,3 +38,22 @@ void FGitSourceControlCommand::DoThreadedWork()
 	Concurrency = EConcurrency::Asynchronous;
 	DoWork();
 }
+
+ECommandResult::Type FGitSourceControlCommand::ReturnResults()
+{
+	// Save any messages that have accumulated
+	for (FString& String : InfoMessages)
+	{
+		Operation->AddInfoMessge(FText::FromString(String));
+	}
+	for (FString& String : ErrorMessages)
+	{
+		Operation->AddErrorMessge(FText::FromString(String));
+	}
+
+	// run the completion delegate if we have one bound
+	ECommandResult::Type Result = bCommandSuccessful ? ECommandResult::Succeeded : ECommandResult::Failed;
+	OperationCompleteDelegate.ExecuteIfBound(Operation, Result);
+
+	return Result;
+}
diff --git a/Engine/Plugins/Developer/GitSourceControl/Source/GitSourceControl/Private/GitSourceControlCommand.h b/Engine/Plugins/Developer/GitSourceControl/Source/GitSourceControl/Private/GitSourceControlCommand.h
--- a/Engine/Plugins/Developer/GitSourceControl/Source/GitSourceControl/Private/GitSourceControlCommand.h
+++ b/Engine/Plugins/Developer/GitSourceControl/Source/GitSourceControl/Private/GitSourceControlCommand.h
@@ -32,9 +32,12 @@ public:
 	/**
 	 * This method is also used to tell the object to cleanup but not before
 	 * the object has finished it's work.
-	 */ 
+	 */
 	virtual void DoThreadedWork() override;
 
+	/** Save any results and call any registered callbacks. */
+	ECommandResult::Type ReturnResults();
+
 public:
 	/** Path to the Git binary */
 	FString PathToGitBinary;
diff --git a/Engine/Plugins/Developer/GitSourceControl/Source/GitSourceControl/Private/GitSourceControlProvider.cpp b/Engine/Plugins/Developer/GitSourceControl/Source/GitSourceControl/Private/GitSourceControlProvider.cpp
--- a/Engine/Plugins/Developer/GitSourceControl/Source/GitSourceControl/Private/GitSourceControlProvider.cpp
+++ b/Engine/Plugins/Developer/GitSourceControl/Source/GitSourceControl/Private/GitSourceControlProvider.cpp
@@ -8,9 +8,11 @@
 #include "Widgets/DeclarativeSyntaxSupport.h"
 #include "GitSourceControlCommand.h"
 #include "ISourceControlModule.h"
+#include "SourceControlHelpers.h"
 #include "GitSourceControlModule.h"
 #include "GitSourceControlUtils.h"
 #include "SGitSourceControlSettings.h"
+#include "SourceControlOperations.h"
 #include "Logging/MessageLog.h"
 #include "ScopedSourceControlProgress.h"
 
@@ -191,6 +193,8 @@ ECommandResult::Type FGitSourceControlProvider::Execute( const TSharedRef<ISourc
 {
 	if(!IsEnabled() && !(InOperation->GetName() == "Connect")) // Only Connect operation allowed while not Enabled (Connected)
 	{
+		// Note that IsEnabled() always returns true so unless it is changed, this code will never be executed
+		InOperationCompleteDelegate.ExecuteIfBound(InOperation, ECommandResult::Failed);
 		return ECommandResult::Failed;
 	}
 
@@ -204,7 +208,11 @@ ECommandResult::Type FGitSourceControlProvider::Execute( const TSharedRef<ISourc
 		FFormatNamedArguments Arguments;
 		Arguments.Add( TEXT("OperationName"), FText::FromName(InOperation->GetName()) );
 		Arguments.Add( TEXT("ProviderName"), FText::FromName(GetName()) );
-		FMessageLog("SourceControl").Error(FText::Format(LOCTEXT("UnsupportedOperation", "Operation '{OperationName}' not supported by source control provider '{ProviderName}'"), Arguments));
+		FText Message(FText::Format(LOCTEXT("UnsupportedOperation", "Operation '{OperationName}' not supported by source control provider '{ProviderName}'"), Arguments));
+		FMessageLog("SourceControl").Error(Message);
+		InOperation->AddErrorMessge(Message);
+
+		InOperationCompleteDelegate.ExecuteIfBound(InOperation, ECommandResult::Failed);
 		return ECommandResult::Failed;
 	}
 
@@ -256,7 +264,7 @@ TSharedPtr<IGitSourceControlWorker, ESPMode::ThreadSafe> FGitSourceControlProvid
 	{
 		return Operation->Execute();
 	}
-		
+
 	return nullptr;
 }
 
@@ -281,7 +289,7 @@ void FGitSourceControlProvider::OutputCommandMessages(const FGitSourceControlCom
 }
 
 void FGitSourceControlProvider::Tick()
-{	
+{
 	bool bStatesUpdated = false;
 	for(int32 CommandIndex = 0; CommandIndex < CommandQueue.Num(); ++CommandIndex)
 	{
@@ -297,9 +305,7 @@ void FGitSourceControlProvider::Tick()
 			// dump any messages to output log
 			OutputCommandMessages(Command);
 
-			// run the completion delegate callback if we have one bound
-			ECommandResult::Type Result = Command.bCommandSuccessful ? ECommandResult::Succeeded : ECommandResult::Failed;
-			Command.OperationCompleteDelegate.ExecuteIfBound(Command.Operation, Result);
+			Command.ReturnResults();
 
 			// commands that are left in the array during a tick need to be deleted
 			if(Command.bAutoDelete)
@@ -307,8 +313,8 @@ void FGitSourceControlProvider::Tick()
 				// Only delete commands that are not running 'synchronously'
 				delete &Command;
 			}
-			
-			// only do one command per tick loop, as we dont want concurrent modification 
+
+			// only do one command per tick loop, as we dont want concurrent modification
 			// of the command queue (which can happen in the completion delegate)
 			break;
 		}
@@ -345,18 +351,18 @@ ECommandResult::Type FGitSourceControlProvider::ExecuteSynchronousCommand(FGitSo
 		// Issue the command asynchronously...
 		IssueCommand( InCommand );
 
-		// ... then wait for its completion (thus making it synchrounous)
+		// ... then wait for its completion (thus making it synchronous)
 		while(!InCommand.bExecuteProcessed)
 		{
 			// Tick the command queue and update progress.
 			Tick();
-			
+
 			Progress.Tick();
 
 			// Sleep for a bit so we don't busy-wait so much.
 			FPlatformProcess::Sleep(0.01f);
 		}
-	
+
 		// always do one more Tick() to make sure the command queue is cleaned up.
 		Tick();
 
@@ -370,7 +376,7 @@ ECommandResult::Type FGitSourceControlProvider::ExecuteSynchronousCommand(FGitSo
 	check(!InCommand.bAutoDelete);
 
 	// ensure commands that are not auto deleted do not end up in the command queue
-	if ( CommandQueue.Contains( &InCommand ) ) 
+	if ( CommandQueue.Contains( &InCommand ) )
 	{
 		CommandQueue.Remove( &InCommand );
 	}
@@ -390,7 +396,13 @@ ECommandResult::Type FGitSourceControlProvider::IssueCommand(FGitSourceControlCo
 	}
 	else
 	{
-		return ECommandResult::Failed;
+		FText Message(LOCTEXT("NoSCCThreads", "There are no threads available to process the source control command."));
+
+		FMessageLog("SourceControl").Error(Message);
+		InCommand.bCommandSuccessful = false;
+		InCommand.Operation->AddErrorMessge(Message);
+
+		return InCommand.ReturnResults();
 	}
 }
 #undef LOCTEXT_NAMESPACE
diff --git a/Engine/Plugins/Developer/PerforceSourceControl/Source/PerforceSourceControl/Private/PerforceSourceControlCommand.cpp b/Engine/Plugins/Developer/PerforceSourceControl/Source/PerforceSourceControl/Private/PerforceSourceControlCommand.cpp
--- a/Engine/Plugins/Developer/PerforceSourceControl/Source/PerforceSourceControl/Private/PerforceSourceControlCommand.cpp
+++ b/Engine/Plugins/Developer/PerforceSourceControl/Source/PerforceSourceControl/Private/PerforceSourceControlCommand.cpp
@@ -71,3 +71,16 @@ bool FPerforceSourceControlCommand::WasCanceledWhileTryingToConnect() const
 {
 	return bCancelledWhileTryingToConnect != 0;
 }
+
+ECommandResult::Type FPerforceSourceControlCommand::ReturnResults()
+{
+	// Save any messages that have accumulated
+	Operation->AppendResultInfo(ResultInfo);
+
+	// run the completion delegate if we have one bound
+	ECommandResult::Type Result = bCancelled ? ECommandResult::Cancelled : (bCommandSuccessful ? ECommandResult::Succeeded : ECommandResult::Failed);
+	OperationCompleteDelegate.ExecuteIfBound(Operation, Result);
+
+	return Result;
+}
+
diff --git a/Engine/Plugins/Developer/PerforceSourceControl/Source/PerforceSourceControl/Private/PerforceSourceControlCommand.h b/Engine/Plugins/Developer/PerforceSourceControl/Source/PerforceSourceControl/Private/PerforceSourceControlCommand.h
--- a/Engine/Plugins/Developer/PerforceSourceControl/Source/PerforceSourceControl/Private/PerforceSourceControlCommand.h
+++ b/Engine/Plugins/Developer/PerforceSourceControl/Source/PerforceSourceControl/Private/PerforceSourceControlCommand.h
@@ -35,7 +35,7 @@ public:
 	/**
 	 * This method is also used to tell the object to cleanup but not before
 	 * the object has finished it's work.
-	 */ 
+	 */
 	virtual void DoThreadedWork() override;
 
 	/** Attempt to cancel the operation */
@@ -56,6 +56,9 @@ public:
 	/** Was the operation canceled while trying to connect? */
 	bool WasCanceledWhileTryingToConnect() const;
 
+	/** Save any results and call any registered callbacks. */
+	ECommandResult::Type ReturnResults();
+
 public:
 	/** Connection parameters, reproduced here because if is not safe to access the provider's settings from another thread */
 	FPerforceConnectionInfo ConnectionInfo;
@@ -96,11 +99,8 @@ public:
 	/** Files to perform this operation on */
 	TArray< FString > Files;
 
-	/**Info and/or warning message message storage*/
-	TArray< FText > InfoMessages;
-
-	/**Potential error message storage*/
-	TArray< FText > ErrorMessages;
+	/** Potential error, warning and info message storage */
+	FSourceControlResultInfo ResultInfo;
 
 	/** Branch names for status queries */
 	TArray< FString > StatusBranchNames;
diff --git a/Engine/Plugins/Developer/PerforceSourceControl/Source/PerforceSourceControl/Private/PerforceSourceControlOperations.cpp b/Engine/Plugins/Developer/PerforceSourceControl/Source/PerforceSourceControl/Private/PerforceSourceControlOperations.cpp
--- a/Engine/Plugins/Developer/PerforceSourceControl/Source/PerforceSourceControl/Private/PerforceSourceControlOperations.cpp
+++ b/Engine/Plugins/Developer/PerforceSourceControl/Source/PerforceSourceControl/Private/PerforceSourceControlOperations.cpp
@@ -61,28 +61,28 @@ struct FBranchModification
 };
 
 
-/** 
+/**
  * Remove redundant errors (that contain a particular string) and also
  * update the commands success status if all errors were removed.
  */
 static void RemoveRedundantErrors(FPerforceSourceControlCommand& InCommand, const FString& InFilter)
 {
 	bool bFoundRedundantError = false;
-	for(auto Iter(InCommand.ErrorMessages.CreateConstIterator()); Iter; Iter++)
+	for(auto Iter(InCommand.ResultInfo.ErrorMessages.CreateConstIterator()); Iter; Iter++)
 	{
 		// Perforce reports files that are already synced as errors, so copy any errors
 		// we get to the info list in this case
 		if(Iter->ToString().Contains(InFilter))
 		{
-			InCommand.InfoMessages.Add(*Iter);
+			InCommand.ResultInfo.InfoMessages.Add(*Iter);
 			bFoundRedundantError = true;
 		}
 	}
 
-	InCommand.ErrorMessages.RemoveAll( FRemoveRedundantErrors(InFilter) );
+	InCommand.ResultInfo.ErrorMessages.RemoveAll( FRemoveRedundantErrors(InFilter) );
 
 	// if we have no error messages now, assume success!
-	if(bFoundRedundantError && InCommand.ErrorMessages.Num() == 0 && !InCommand.bCommandSuccessful)
+	if(bFoundRedundantError && InCommand.ResultInfo.ErrorMessages.Num() == 0 && !InCommand.bCommandSuccessful)
 	{
 		InCommand.bCommandSuccessful = true;
 	}
@@ -179,26 +179,26 @@ static bool CheckWorkspaceRecordSet(const FP4RecordSet& InRecords, TArray<FText>
 
 	for(const auto& Record : InRecords)
 	{
-		FString Root = Record(TEXT("Root"));	
-		
+		FString Root = Record(TEXT("Root"));
+
 		// A workspace root could be "null" which allows the user to map depot locations to different drives.
 		// Allow these workspaces since we already allow workspaces mapped to drive letters.
 		const bool bIsNullClientRootPath = (Root == TEXT("null"));
-			
+
 		// Sanitize root name
 		Root = Root.Replace(TEXT("\\"), TEXT("/"));
 		if (!Root.EndsWith(TEXT("/")))
 		{
 			Root += TEXT("/");
 		}
-			
+
 		if (bIsNullClientRootPath || ApplicationPath.Contains(Root))
 		{
 			return true;
 		}
 		else
 		{
-			const FString Client = Record(TEXT("Client"));	
+			const FString Client = Record(TEXT("Client"));
 			OutNotificationText = FText::Format(LOCTEXT("WorkspaceError", "Workspace '{0}' does not map into this project's directory."), FText::FromString(Client));
 			OutErrorMessages.Add(OutNotificationText);
 			OutErrorMessages.Add(LOCTEXT("WorkspaceHelp", "You should set your workspace up to map to a directory at or above the project's directory."));
@@ -236,21 +236,21 @@ bool FPerforceConnectWorker::Execute(FPerforceSourceControlCommand& InCommand)
 		FP4RecordSet Records;
 		Parameters.Add(TEXT("-o"));
 		Parameters.Add(InCommand.ConnectionInfo.Workspace);
-		InCommand.bCommandSuccessful = Connection.RunCommand(TEXT("client"), Parameters, Records, InCommand.ErrorMessages, FOnIsCancelled::CreateRaw(&InCommand, &FPerforceSourceControlCommand::IsCanceled), InCommand.bConnectionDropped);
-		
+		InCommand.bCommandSuccessful = Connection.RunCommand(TEXT("client"), Parameters, Records, InCommand.ResultInfo.ErrorMessages, FOnIsCancelled::CreateRaw(&InCommand, &FPerforceSourceControlCommand::IsCanceled), InCommand.bConnectionDropped);
+
 		// If there are error messages, user name is most likely invalid. Otherwise, make sure workspace actually
 		// exists on server by checking if we have it's update date.
-		InCommand.bCommandSuccessful &= InCommand.ErrorMessages.Num() == 0 && Records.Num() > 0 && Records[0].Contains(TEXT("Update"));
-		if (!InCommand.bCommandSuccessful && InCommand.ErrorMessages.Num() == 0)
+		InCommand.bCommandSuccessful &= InCommand.ResultInfo.ErrorMessages.Num() == 0 && Records.Num() > 0 && Records[0].Contains(TEXT("Update"));
+		if (!InCommand.bCommandSuccessful && InCommand.ResultInfo.ErrorMessages.Num() == 0)
 		{
-			InCommand.ErrorMessages.Add(LOCTEXT("InvalidWorkspace", "Invalid workspace."));
+			InCommand.ResultInfo.ErrorMessages.Add(LOCTEXT("InvalidWorkspace", "Invalid workspace."));
 		}
 
 		// check if we can actually work with this workspace
 		if(InCommand.bCommandSuccessful)
 		{
 			FText Notification;
-			InCommand.bCommandSuccessful = CheckWorkspaceRecordSet(Records, InCommand.ErrorMessages, Notification);
+			InCommand.bCommandSuccessful = CheckWorkspaceRecordSet(Records, InCommand.ResultInfo.ErrorMessages, Notification);
 			if(!InCommand.bCommandSuccessful)
 			{
 				check(InCommand.Operation->GetName() == GetName());
@@ -261,7 +261,7 @@ bool FPerforceConnectWorker::Execute(FPerforceSourceControlCommand& InCommand)
 
 		if(InCommand.bCommandSuccessful)
 		{
-			ParseRecordSet(Records, InCommand.InfoMessages);
+			ParseRecordSet(Records, InCommand.ResultInfo.InfoMessages);
 		}
 	}
 	return InCommand.bCommandSuccessful;
@@ -278,7 +278,7 @@ FName FPerforceCheckOutWorker::GetName() const
 }
 
 bool FPerforceCheckOutWorker::Execute(FPerforceSourceControlCommand& InCommand)
-{	
+{
 	FScopedPerforceConnection ScopedConnection(InCommand);
 	if (!InCommand.IsCanceled() && ScopedConnection.IsValid())
 	{
@@ -289,7 +289,7 @@ bool FPerforceCheckOutWorker::Execute(FPerforceSourceControlCommand& InCommand)
 
 		Parameters.Append(InCommand.Files);
 		FP4RecordSet Records;
-		InCommand.bCommandSuccessful = Connection.RunCommand(TEXT("edit"), Parameters, Records, InCommand.ErrorMessages, FOnIsCancelled::CreateRaw(&InCommand, &FPerforceSourceControlCommand::IsCanceled), InCommand.bConnectionDropped);
+		InCommand.bCommandSuccessful = Connection.RunCommand(TEXT("edit"), Parameters, Records, InCommand.ResultInfo.ErrorMessages, FOnIsCancelled::CreateRaw(&InCommand, &FPerforceSourceControlCommand::IsCanceled), InCommand.bConnectionDropped);
 		ParseRecordSetForState(Records, OutResults);
 	}
 
@@ -326,13 +326,13 @@ bool FPerforceCheckInWorker::Execute(FPerforceSourceControlCommand& InCommand)
 {
 	FScopedPerforceConnection ScopedConnection(InCommand);
 	if (!InCommand.IsCanceled() && ScopedConnection.IsValid())
-	{	
+	{
 		FPerforceConnection& Connection = ScopedConnection.GetConnection();
 
 		check(InCommand.Operation->GetName() == GetName());
 		TSharedRef<FCheckIn, ESPMode::ThreadSafe> Operation = StaticCastSharedRef<FCheckIn>(InCommand.Operation);
 
-		int32 ChangeList = Connection.CreatePendingChangelist(Operation->GetDescription(), FOnIsCancelled::CreateRaw(&InCommand, &FPerforceSourceControlCommand::IsCanceled), InCommand.ErrorMessages);
+		int32 ChangeList = Connection.CreatePendingChangelist(Operation->GetDescription(), FOnIsCancelled::CreateRaw(&InCommand, &FPerforceSourceControlCommand::IsCanceled), InCommand.ResultInfo.ErrorMessages);
 		if (ChangeList > 0)
 		{
 			// Batch reopen into multiple commands, to avoid command line limits
@@ -342,7 +342,7 @@ bool FPerforceCheckInWorker::Execute(FPerforceSourceControlCommand& InCommand)
 			{
 				FP4RecordSet Records;
 				TArray< FString > ReopenParams;
-						
+
 				//Add changelist information to params
 				ReopenParams.Insert(TEXT("-c"), 0);
 				ReopenParams.Insert(FString::Printf(TEXT("%d"), ChangeList), 1);
@@ -353,7 +353,7 @@ bool FPerforceCheckInWorker::Execute(FPerforceSourceControlCommand& InCommand)
 					ReopenParams.Add(InCommand.Files[FileIndex]);
 				}
 
-				InCommand.bCommandSuccessful = Connection.RunCommand(TEXT("reopen"), ReopenParams, Records, InCommand.ErrorMessages, FOnIsCancelled::CreateRaw(&InCommand, &FPerforceSourceControlCommand::IsCanceled), InCommand.bConnectionDropped);
+				InCommand.bCommandSuccessful = Connection.RunCommand(TEXT("reopen"), ReopenParams, Records, InCommand.ResultInfo.ErrorMessages, FOnIsCancelled::CreateRaw(&InCommand, &FPerforceSourceControlCommand::IsCanceled), InCommand.bConnectionDropped);
 			}
 
 			if (InCommand.bCommandSuccessful)
@@ -365,9 +365,9 @@ bool FPerforceCheckInWorker::Execute(FPerforceSourceControlCommand& InCommand)
 				SubmitParams.Insert(TEXT("-c"), 0);
 				SubmitParams.Insert(FString::Printf(TEXT("%d"), ChangeList), 1);
 
-				InCommand.bCommandSuccessful = Connection.RunCommand(TEXT("submit"), SubmitParams, Records, InCommand.ErrorMessages, FOnIsCancelled::CreateRaw(&InCommand, &FPerforceSourceControlCommand::IsCanceled), InCommand.bConnectionDropped);
+				InCommand.bCommandSuccessful = Connection.RunCommand(TEXT("submit"), SubmitParams, Records, InCommand.ResultInfo.ErrorMessages, FOnIsCancelled::CreateRaw(&InCommand, &FPerforceSourceControlCommand::IsCanceled), InCommand.bConnectionDropped);
 
-				if (InCommand.ErrorMessages.Num() > 0)
+				if (InCommand.ResultInfo.ErrorMessages.Num() > 0)
 				{
 					InCommand.bCommandSuccessful = false;
 				}
@@ -420,7 +420,7 @@ FName FPerforceMarkForAddWorker::GetName() const
 bool FPerforceMarkForAddWorker::Execute(FPerforceSourceControlCommand& InCommand)
 {
 	// Perforce will allow you to mark files for add that don't currently exist on disk
-	// This goes against the workflow of our other SCC providers (such as SVN and Git), 
+	// This goes against the workflow of our other SCC providers (such as SVN and Git),
 	// so we manually check that the files exist before allowing this command to continue
 	// This keeps the behavior consistent between SCC providers
 	bool bHasMissingFiles = false;
@@ -429,7 +429,7 @@ bool FPerforceMarkForAddWorker::Execute(FPerforceSourceControlCommand& InCommand
 		if(!IFileManager::Get().FileExists(*FileToAdd))
 		{
 			bHasMissingFiles = true;
-			InCommand.ErrorMessages.Add(FText::Format(LOCTEXT("Error_FailedToMarkFileForAdd_FileMissing", "Failed mark the file '{0}' for add. The file doesn't exist on disk."), FText::FromString(FileToAdd)));
+			InCommand.ResultInfo.ErrorMessages.Add(FText::Format(LOCTEXT("Error_FailedToMarkFileForAdd_FileMissing", "Failed mark the file '{0}' for add. The file doesn't exist on disk."), FText::FromString(FileToAdd)));
 		}
 	}
 	if(bHasMissingFiles)
@@ -448,7 +448,7 @@ bool FPerforceMarkForAddWorker::Execute(FPerforceSourceControlCommand& InCommand
 		AppendChangelistParameter(Parameters);
 		Parameters.Append(InCommand.Files);
 
-		InCommand.bCommandSuccessful = Connection.RunCommand(TEXT("add"), Parameters, Records, InCommand.ErrorMessages, FOnIsCancelled::CreateRaw(&InCommand, &FPerforceSourceControlCommand::IsCanceled), InCommand.bConnectionDropped);
+		InCommand.bCommandSuccessful = Connection.RunCommand(TEXT("add"), Parameters, Records, InCommand.ResultInfo.ErrorMessages, FOnIsCancelled::CreateRaw(&InCommand, &FPerforceSourceControlCommand::IsCanceled), InCommand.bConnectionDropped);
 		ParseRecordSetForState(Records, OutResults);
 	}
 	return InCommand.bCommandSuccessful;
@@ -476,7 +476,7 @@ bool FPerforceDeleteWorker::Execute(FPerforceSourceControlCommand& InCommand)
 		Parameters.Append(InCommand.Files);
 
 		FP4RecordSet Records;
-		InCommand.bCommandSuccessful = Connection.RunCommand(TEXT("delete"), Parameters, Records, InCommand.ErrorMessages, FOnIsCancelled::CreateRaw(&InCommand, &FPerforceSourceControlCommand::IsCanceled), InCommand.bConnectionDropped);
+		InCommand.bCommandSuccessful = Connection.RunCommand(TEXT("delete"), Parameters, Records, InCommand.ResultInfo.ErrorMessages, FOnIsCancelled::CreateRaw(&InCommand, &FPerforceSourceControlCommand::IsCanceled), InCommand.bConnectionDropped);
 		ParseRecordSetForState(Records, OutResults);
 	}
 	return InCommand.bCommandSuccessful;
@@ -504,7 +504,7 @@ bool FPerforceRevertWorker::Execute(FPerforceSourceControlCommand& InCommand)
 		Parameters.Append(InCommand.Files);
 
 		FP4RecordSet Records;
-		InCommand.bCommandSuccessful = Connection.RunCommand(TEXT("revert"), Parameters, Records, InCommand.ErrorMessages, FOnIsCancelled::CreateRaw(&InCommand, &FPerforceSourceControlCommand::IsCanceled), InCommand.bConnectionDropped);
+		InCommand.bCommandSuccessful = Connection.RunCommand(TEXT("revert"), Parameters, Records, InCommand.ResultInfo.ErrorMessages, FOnIsCancelled::CreateRaw(&InCommand, &FPerforceSourceControlCommand::IsCanceled), InCommand.bConnectionDropped);
 		ParseRecordSetForState(Records, OutResults);
 	}
 	return InCommand.bCommandSuccessful;
@@ -562,7 +562,7 @@ bool FPerforceSyncWorker::Execute(FPerforceSourceControlCommand& InCommand)
 		}
 
 		FP4RecordSet Records;
-		InCommand.bCommandSuccessful = Connection.RunCommand(TEXT("sync"), Parameters, Records, InCommand.ErrorMessages, FOnIsCancelled::CreateRaw(&InCommand, &FPerforceSourceControlCommand::IsCanceled), InCommand.bConnectionDropped);
+		InCommand.bCommandSuccessful = Connection.RunCommand(TEXT("sync"), Parameters, Records, InCommand.ResultInfo.ErrorMessages, FOnIsCancelled::CreateRaw(&InCommand, &FPerforceSourceControlCommand::IsCanceled), InCommand.bConnectionDropped);
 		ParseSyncResults(Records, OutResults);
 
 		RemoveRedundantErrors(InCommand, TEXT("file(s) up-to-date"));
@@ -654,7 +654,7 @@ static void ParseBranchModificationResults(const FP4RecordSet& InRecords, const
 
 static void ParseUpdateStatusResults(const FP4RecordSet& InRecords, const TArray<FText>& ErrorMessages, TArray<FPerforceSourceControlState>& OutStates, const FString& ContentRoot, TMap<FString, FBranchModification>& BranchModifications)
 {
-	// Build up a map of any other branch states	
+	// Build up a map of any other branch states
 	for (int32 Index = 0; Index < InRecords.Num(); ++Index)
 	{
 		const FP4Record& ClientRecord = InRecords[Index];
@@ -693,7 +693,7 @@ static void ParseUpdateStatusResults(const FP4RecordSet& InRecords, const TArray
 
 		// Store checkout information to branch state
 		FBranchModification& BranchModification = BranchModifications[FileName];
-	
+
 		if (OtherOpen.Len())
 		{
 			BranchModification.CheckedOutBranches.AddUnique(Branch);
@@ -743,17 +743,17 @@ static void ParseUpdateStatusResults(const FP4RecordSet& InRecords, const TArray
 		State.DepotFilename = DepotFileName;
 
 		State.State = EPerforceState::ReadOnly;
-		if (Action.Len() > 0 && Action == TEXT("add")) 
+		if (Action.Len() > 0 && Action == TEXT("add"))
 		{
 			State.State = EPerforceState::OpenForAdd;
 		}
-		else if (Action.Len() > 0 && Action == TEXT("delete")) 
+		else if (Action.Len() > 0 && Action == TEXT("delete"))
 		{
 			State.State = EPerforceState::MarkedForDelete;
 		}
 		else if (OpenType.Len() > 0)
 		{
-			if(Action.Len() > 0 && Action == TEXT("branch")) 
+			if(Action.Len() > 0 && Action == TEXT("branch"))
 			{
 				State.State = EPerforceState::Branched;
 			}
@@ -780,11 +780,11 @@ static void ParseUpdateStatusResults(const FP4RecordSet& InRecords, const TArray
 
 			// Add to the checked out branches
 			State.CheckedOutBranches.AddUnique(FEngineVersion::Current().GetBranch());
-				
+
 			State.State = EPerforceState::CheckedOutOther;
 		}
 		//file has been previously deleted, ok to add again
-		else if (HeadAction.Len() > 0 && HeadAction == TEXT("delete")) 
+		else if (HeadAction.Len() > 0 && HeadAction == TEXT("delete"))
 		{
 			State.State = EPerforceState::NotInDepot;
 		}
@@ -814,7 +814,7 @@ static void ParseUpdateStatusResults(const FP4RecordSet& InRecords, const TArray
 					State.HeadChangeList = BranchModification.ChangeList;
 				}
 			}
-			
+
 			// Setup other branch check outs
 			if (BranchModification.CheckedOutBranches.Num())
 			{
@@ -826,7 +826,7 @@ static void ParseUpdateStatusResults(const FP4RecordSet& InRecords, const TArray
 				}
 			}
 		}
-		
+
 		if (HeadRev.Len() > 0 && HaveRev.Len() > 0)
 		{
 			TTypeFromString<int>::FromString(State.DepotRevNumber, *HeadRev);
@@ -858,7 +858,7 @@ static void ParseUpdateStatusResults(const FP4RecordSet& InRecords, const TArray
 					FString ResolveBaseRev = ClientRecord(VarName);
 
 					TTypeFromString<int>::FromString(State.PendingResolveRevNumber, *ResolveBaseRev);
-					
+
 					++ResolveActionNumber;
 				}
 			}
@@ -906,7 +906,7 @@ static void ParseUpdateStatusResults(const FP4RecordSet& InRecords, const TArray
 			FPaths::NormalizeFilename(FullPath);
 			OutStates.Add(FPerforceSourceControlState(FullPath));
 			FPerforceSourceControlState& State = OutStates.Last();
-			State.State = EPerforceState::NotUnderClientRoot;	
+			State.State = EPerforceState::NotUnderClientRoot;
 		}
 	}
 }
@@ -987,7 +987,7 @@ static void ParseHistoryResults(const FP4RecordSet& InRecords, const TArray<FPer
 		{
 			const FP4Record& ClientRecord = InRecords[RecordIndex];
 
-			// Extract the file name 
+			// Extract the file name
 			check(ClientRecord.Contains(TEXT("depotFile")));
 			FString DepotFileName = ClientRecord(TEXT("depotFile"));
 			FString LocalFileName = FindWorkspaceFile(InStates, DepotFileName);
@@ -1039,7 +1039,7 @@ static void ParseHistoryResults(const FP4RecordSet& InRecords, const TArray<FPer
 					check(ClientRecord.Contains(*VarName));
 					FileSize = ClientRecord(*VarName);
 				}
-		
+
 				// Extract the clientspec/workspace
 				VarName = FString::Printf(TEXT("client%d"), RevisionNumbers);
 				check(ClientRecord.Contains(*VarName));
@@ -1097,7 +1097,7 @@ static void ParseDiffResults(const FP4RecordSet& InRecords, TArray<FString>& Out
 			const FP4Record& ClientRecord = InRecords[Index];
 			FString FileName = ClientRecord(TEXT("clientFile"));
 			FPaths::NormalizeFilename(FileName);
-			OutModifiedFiles.Add(FileName);	
+			OutModifiedFiles.Add(FileName);
 		}
 	}
 }
@@ -1128,7 +1128,7 @@ bool FPerforceUpdateStatusWorker::Execute(FPerforceSourceControlCommand& InComma
 			// Get the branches of interest for status updates
 			const FString& ContentRoot = InCommand.ContentRoot;
 			const TArray<FString>& StatusBranches = InCommand.StatusBranchNames;
-			
+
 			// Mandatory parameters (the list of files to stat):
 			for (FString File : InCommand.Files)
 			{
@@ -1145,7 +1145,7 @@ bool FPerforceUpdateStatusWorker::Execute(FPerforceSourceControlCommand& InComma
 						FString BranchFile;
 						if (File.Split(ContentRoot, nullptr, &BranchFile))
 						{
-							TArray<FStringFormatArg> Args = { Branch, ContentRoot, BranchFile };							
+							TArray<FStringFormatArg> Args = { Branch, ContentRoot, BranchFile };
 							Parameters.Add(FString::Format(TEXT("{0}/{1}{2}"), Args));
 						}
 					}
@@ -1161,7 +1161,7 @@ bool FPerforceUpdateStatusWorker::Execute(FPerforceSourceControlCommand& InComma
 			TMap<FString, FBranchModification> BranchModifications;
 			if (StatusBranches.Num())
 			{
-				// Get all revisions to check for modifications on other branches			
+				// Get all revisions to check for modifications on other branches
 				TArray<FString> RevisionParameters = Parameters;
 				// Sort by head revision
 				RevisionParameters.Insert(TEXT("-Sr"), 0);
@@ -1169,13 +1169,13 @@ bool FPerforceUpdateStatusWorker::Execute(FPerforceSourceControlCommand& InComma
 				RevisionParameters.Insert(TEXT("-Of"), 0);
 
 				FP4RecordSet RevisionRecords;
-				InCommand.bCommandSuccessful &= Connection.RunCommand(TEXT("fstat"), RevisionParameters, RevisionRecords, InCommand.ErrorMessages, FOnIsCancelled::CreateRaw(&InCommand, &FPerforceSourceControlCommand::IsCanceled), InCommand.bConnectionDropped);
-				ParseBranchModificationResults(RevisionRecords, InCommand.ErrorMessages, ContentRoot, BranchModifications);
-			}			
-			
+				InCommand.bCommandSuccessful &= Connection.RunCommand(TEXT("fstat"), RevisionParameters, RevisionRecords, InCommand.ResultInfo.ErrorMessages, FOnIsCancelled::CreateRaw(&InCommand, &FPerforceSourceControlCommand::IsCanceled), InCommand.bConnectionDropped);
+				ParseBranchModificationResults(RevisionRecords, InCommand.ResultInfo.ErrorMessages, ContentRoot, BranchModifications);
+			}
+
 			FP4RecordSet Records;
-			InCommand.bCommandSuccessful &= Connection.RunCommand(TEXT("fstat"), Parameters, Records, InCommand.ErrorMessages, FOnIsCancelled::CreateRaw(&InCommand, &FPerforceSourceControlCommand::IsCanceled), InCommand.bConnectionDropped);
-			ParseUpdateStatusResults(Records, InCommand.ErrorMessages, OutStates, ContentRoot, BranchModifications);
+			InCommand.bCommandSuccessful &= Connection.RunCommand(TEXT("fstat"), Parameters, Records, InCommand.ResultInfo.ErrorMessages, FOnIsCancelled::CreateRaw(&InCommand, &FPerforceSourceControlCommand::IsCanceled), InCommand.bConnectionDropped);
+			ParseUpdateStatusResults(Records, InCommand.ResultInfo.ErrorMessages, OutStates, ContentRoot, BranchModifications);
 			RemoveRedundantErrors(InCommand, TEXT(" - no such file(s)."));
 			RemoveRedundantErrors(InCommand, TEXT("' is not under client's root '"));
 		}
@@ -1203,7 +1203,7 @@ bool FPerforceUpdateStatusWorker::Execute(FPerforceSourceControlCommand& InComma
 			//limit to last 100 changes
 			Parameters.Add(TEXT("-m 100"));
 			Parameters.Append(InCommand.Files);
-			InCommand.bCommandSuccessful &= Connection.RunCommand(TEXT("filelog"), Parameters, Records, InCommand.ErrorMessages, FOnIsCancelled::CreateRaw(&InCommand, &FPerforceSourceControlCommand::IsCanceled), InCommand.bConnectionDropped);
+			InCommand.bCommandSuccessful &= Connection.RunCommand(TEXT("filelog"), Parameters, Records, InCommand.ResultInfo.ErrorMessages, FOnIsCancelled::CreateRaw(&InCommand, &FPerforceSourceControlCommand::IsCanceled), InCommand.bConnectionDropped);
 			ParseHistoryResults(Records, OutStates, OutHistory);
 			RemoveRedundantErrors(InCommand, TEXT(" - no such file(s)."));
 			RemoveRedundantErrors(InCommand, TEXT(" - file(s) not on client"));
@@ -1217,7 +1217,7 @@ bool FPerforceUpdateStatusWorker::Execute(FPerforceSourceControlCommand& InComma
 			TArray<FString> Parameters = InCommand.Files;
 			Parameters.Add(FileQuery);
 			FP4RecordSet Records;
-			InCommand.bCommandSuccessful &= Connection.RunCommand(TEXT("opened"), Parameters, Records, InCommand.ErrorMessages, FOnIsCancelled::CreateRaw(&InCommand, &FPerforceSourceControlCommand::IsCanceled), InCommand.bConnectionDropped);
+			InCommand.bCommandSuccessful &= Connection.RunCommand(TEXT("opened"), Parameters, Records, InCommand.ResultInfo.ErrorMessages, FOnIsCancelled::CreateRaw(&InCommand, &FPerforceSourceControlCommand::IsCanceled), InCommand.bConnectionDropped);
 			ParseOpenedResults(Records, ANSI_TO_TCHAR(Connection.P4Client.GetClient().Text()), Connection.ClientRoot, OutStateMap);
 			RemoveRedundantErrors(InCommand, TEXT(" - no such file(s)."));
 			RemoveRedundantErrors(InCommand, TEXT("' is not under client's root '"));
@@ -1230,13 +1230,15 @@ bool FPerforceUpdateStatusWorker::Execute(FPerforceSourceControlCommand& InComma
 			// Query for open files different than the versions stored in Perforce
 			Parameters.Add(TEXT("-sa"));
 			Parameters.Append(InCommand.Files);
-			InCommand.bCommandSuccessful &= Connection.RunCommand(TEXT("diff"), Parameters, Records, InCommand.ErrorMessages, FOnIsCancelled::CreateRaw(&InCommand, &FPerforceSourceControlCommand::IsCanceled), InCommand.bConnectionDropped);
+			InCommand.bCommandSuccessful &= Connection.RunCommand(TEXT("diff"), Parameters, Records, InCommand.ResultInfo.ErrorMessages, FOnIsCancelled::CreateRaw(&InCommand, &FPerforceSourceControlCommand::IsCanceled), InCommand.bConnectionDropped);
 
 			// Parse the results and store them in the command
 			ParseDiffResults(Records, OutModifiedFiles);
 			RemoveRedundantErrors(InCommand, TEXT(" - no such file(s)."));
 			RemoveRedundantErrors(InCommand, TEXT(" - file(s) not opened for edit"));
 			RemoveRedundantErrors(InCommand, TEXT("' is not under client's root '"));
+			RemoveRedundantErrors(InCommand, TEXT(" - file(s) not opened on this client"));
+			RemoveRedundantErrors(InCommand, TEXT(" - protected namespace - access denied"));
 		}
 	}
 #endif
@@ -1259,7 +1261,7 @@ bool FPerforceUpdateStatusWorker::UpdateStates() const
 		// Update every member except History and Timestamp. History will be updated below from the OutHistory map.
 		// Timestamp is used to throttle status requests, so update it to current time:
 		auto History = MoveTemp(State->History);
-		*State = Status; 
+		*State = Status;
 		State->History = MoveTemp(History);
 		State->TimeStamp = Now;
 		bUpdated = true;
@@ -1303,7 +1305,7 @@ bool FPerforceGetWorkspacesWorker::Execute(FPerforceSourceControlCommand& InComm
 	{
 		FPerforceConnection& Connection = ScopedConnection.GetConnection();
 		TArray<FString> ClientSpecList;
-		InCommand.bCommandSuccessful = Connection.GetWorkspaceList(InCommand.ConnectionInfo, FOnIsCancelled::CreateRaw(&InCommand, &FPerforceSourceControlCommand::IsCanceled), ClientSpecList, InCommand.ErrorMessages);
+		InCommand.bCommandSuccessful = Connection.GetWorkspaceList(InCommand.ConnectionInfo, FOnIsCancelled::CreateRaw(&InCommand, &FPerforceSourceControlCommand::IsCanceled), ClientSpecList, InCommand.ResultInfo.ErrorMessages);
 
 		check(InCommand.Operation->GetName() == GetName());
 		TSharedRef<FGetWorkspaces, ESPMode::ThreadSafe> Operation = StaticCastSharedRef<FGetWorkspaces>(InCommand.Operation);
@@ -1341,9 +1343,9 @@ bool FPerforceCopyWorker::Execute(class FPerforceSourceControlCommand& InCommand
 
 		Parameters.Append(InCommand.Files);
 		Parameters.Add(DestinationPath);
-		
+
 		FP4RecordSet Records;
-		InCommand.bCommandSuccessful = Connection.RunCommand(TEXT("integrate"), Parameters, Records, InCommand.ErrorMessages, FOnIsCancelled::CreateRaw(&InCommand, &FPerforceSourceControlCommand::IsCanceled), InCommand.bConnectionDropped);
+		InCommand.bCommandSuccessful = Connection.RunCommand(TEXT("integrate"), Parameters, Records, InCommand.ResultInfo.ErrorMessages, FOnIsCancelled::CreateRaw(&InCommand, &FPerforceSourceControlCommand::IsCanceled), InCommand.bConnectionDropped);
 
 		// We now need to do a p4 resolve.
 		// This is because when we copy a file in the Editor, we first make the copy on disk before attempting to branch. This causes a conflict in P4's eyes.
@@ -1354,7 +1356,7 @@ bool FPerforceCopyWorker::Execute(class FPerforceSourceControlCommand& InCommand
 			TArray<FString> ResolveParameters;
 			ResolveParameters.Add(TEXT("-ay"));	// 'accept yours'
 			ResolveParameters.Add(DestinationPath);
-			InCommand.bCommandSuccessful = Connection.RunCommand(TEXT("resolve"), ResolveParameters, Records, InCommand.ErrorMessages, FOnIsCancelled::CreateRaw(&InCommand, &FPerforceSourceControlCommand::IsCanceled), InCommand.bConnectionDropped);
+			InCommand.bCommandSuccessful = Connection.RunCommand(TEXT("resolve"), ResolveParameters, Records, InCommand.ResultInfo.ErrorMessages, FOnIsCancelled::CreateRaw(&InCommand, &FPerforceSourceControlCommand::IsCanceled), InCommand.bConnectionDropped);
 		}
 	}
 	return InCommand.bCommandSuccessful;
@@ -1366,12 +1368,12 @@ bool FPerforceCopyWorker::UpdateStates() const
 }
 
 // IPerforceSourceControlWorker interface
-FName FPerforceResolveWorker::GetName() const 
+FName FPerforceResolveWorker::GetName() const
 {
 	return "Resolve";
 }
 
-bool FPerforceResolveWorker::Execute(class FPerforceSourceControlCommand& InCommand) 
+bool FPerforceResolveWorker::Execute(class FPerforceSourceControlCommand& InCommand)
 {
 	FScopedPerforceConnection ScopedConnection(InCommand);
 	if (!InCommand.IsCanceled() && ScopedConnection.IsValid())
@@ -1385,7 +1387,7 @@ bool FPerforceResolveWorker::Execute(class FPerforceSourceControlCommand& InComm
 		AppendChangelistParameter(Parameters);
 
 		FP4RecordSet Records;
-		InCommand.bCommandSuccessful = Connection.RunCommand(TEXT("resolve"), Parameters, Records, InCommand.ErrorMessages, FOnIsCancelled::CreateRaw(&InCommand, &FPerforceSourceControlCommand::IsCanceled), InCommand.bConnectionDropped);
+		InCommand.bCommandSuccessful = Connection.RunCommand(TEXT("resolve"), Parameters, Records, InCommand.ResultInfo.ErrorMessages, FOnIsCancelled::CreateRaw(&InCommand, &FPerforceSourceControlCommand::IsCanceled), InCommand.bConnectionDropped);
 		if( InCommand.bCommandSuccessful )
 		{
 			UpdatedFiles = InCommand.Files;
@@ -1395,7 +1397,7 @@ bool FPerforceResolveWorker::Execute(class FPerforceSourceControlCommand& InComm
 	return InCommand.bCommandSuccessful;
 }
 
-bool FPerforceResolveWorker::UpdateStates() const 
+bool FPerforceResolveWorker::UpdateStates() const
 {
 	FPerforceSourceControlModule& PerforceSourceControl = FPerforceSourceControlModule::Get();
 
diff --git a/Engine/Plugins/Developer/PerforceSourceControl/Source/PerforceSourceControl/Private/PerforceSourceControlProvider.cpp b/Engine/Plugins/Developer/PerforceSourceControl/Source/PerforceSourceControl/Private/PerforceSourceControlProvider.cpp
--- a/Engine/Plugins/Developer/PerforceSourceControl/Source/PerforceSourceControl/Private/PerforceSourceControlProvider.cpp
+++ b/Engine/Plugins/Developer/PerforceSourceControl/Source/PerforceSourceControl/Private/PerforceSourceControlProvider.cpp
@@ -18,6 +18,8 @@
 #include "SPerforceSourceControlSettings.h"
 #include "Logging/MessageLog.h"
 #include "ScopedSourceControlProgress.h"
+#include "SourceControlHelpers.h"
+#include "SourceControlOperations.h"
 
 static FName ProviderName("Perforce");
 
@@ -142,7 +144,7 @@ void FPerforceSourceControlProvider::ParseCommandLineSettings(bool bForceConnect
 		P4Settings.SetHostOverride(HostOverrideName);
 		P4Settings.SetChangelistNumber(Changelist);
 	}
-	
+
 	if (bForceConnection)
 	{
 		FPerforceConnectionInfo ConnectionInfo = P4Settings.GetConnectionInfo();
@@ -152,9 +154,8 @@ void FPerforceSourceControlProvider::ParseCommandLineSettings(bool bForceConnect
 			P4Settings.SetUserName(UserName);
 			P4Settings.SetWorkspace(ClientSpecName);
 			P4Settings.SetHostOverride(HostOverrideName);
+			bServerAvailable = true;
 		}
-
-		bServerAvailable = true;
 	}
 
 	//Save off settings so this doesn't happen every time
@@ -169,9 +170,9 @@ void FPerforceSourceControlProvider::GetWorkspaceList(const FPerforceConnectionI
 	Connection.GetWorkspaceList(InConnectionInfo, FOnIsCancelled(), OutWorkspaceList, OutErrorMessages);
 }
 
-const FString& FPerforceSourceControlProvider::GetTicket() const 
-{ 
-	return Ticket; 
+const FString& FPerforceSourceControlProvider::GetTicket() const
+{
+	return Ticket;
 }
 
 const FName& FPerforceSourceControlProvider::GetName() const
@@ -246,6 +247,8 @@ ECommandResult::Type FPerforceSourceControlProvider::Execute( const TSharedRef<I
 {
 	if(!IsEnabled())
 	{
+		// Note that IsEnabled() always returns true so unless it is changed, this code will never be executed
+		InOperationCompleteDelegate.ExecuteIfBound(InOperation, ECommandResult::Failed);
 		return ECommandResult::Failed;
 	}
 
@@ -259,7 +262,12 @@ ECommandResult::Type FPerforceSourceControlProvider::Execute( const TSharedRef<I
 		FFormatNamedArguments Arguments;
 		Arguments.Add( TEXT("OperationName"), FText::FromName(InOperation->GetName()) );
 		Arguments.Add( TEXT("ProviderName"), FText::FromName(GetName()) );
-		FMessageLog("SourceControl").Error(FText::Format(LOCTEXT("UnsupportedOperation", "Operation '{OperationName}' not supported by source control provider '{ProviderName}'"), Arguments));
+		FText Message = FText::Format(LOCTEXT("UnsupportedOperation", "Operation '{OperationName}' not supported by source control provider '{ProviderName}'"), Arguments);
+
+		FMessageLog("SourceControl").Error(Message);
+		InOperation->AddErrorMessge(Message);
+
+		InOperationCompleteDelegate.ExecuteIfBound(InOperation, ECommandResult::Failed);
 		return ECommandResult::Failed;
 	}
 
@@ -335,19 +343,19 @@ void FPerforceSourceControlProvider::OutputCommandMessages(const FPerforceSource
 {
 	FMessageLog SourceControlLog("SourceControl");
 
-	for (int32 ErrorIndex = 0; ErrorIndex < InCommand.ErrorMessages.Num(); ++ErrorIndex)
+	for (int32 ErrorIndex = 0; ErrorIndex < InCommand.ResultInfo.ErrorMessages.Num(); ++ErrorIndex)
 	{
-		SourceControlLog.Error(InCommand.ErrorMessages[ErrorIndex]);
+		SourceControlLog.Error(InCommand.ResultInfo.ErrorMessages[ErrorIndex]);
 	}
 
-	for (int32 InfoIndex = 0; InfoIndex < InCommand.InfoMessages.Num(); ++InfoIndex)
+	for (int32 InfoIndex = 0; InfoIndex < InCommand.ResultInfo.InfoMessages.Num(); ++InfoIndex)
 	{
-		SourceControlLog.Info(InCommand.InfoMessages[InfoIndex]);
+		SourceControlLog.Info(InCommand.ResultInfo.InfoMessages[InfoIndex]);
 	}
 }
 
 void FPerforceSourceControlProvider::Tick()
-{	
+{
 	bool bStatesUpdated = false;
 	for (int32 CommandIndex = 0; CommandIndex < CommandQueue.Num(); ++CommandIndex)
 	{
@@ -370,17 +378,7 @@ void FPerforceSourceControlProvider::Tick()
 			// have been called. Otherwise, now we have to call it.
 			if (!Command.bCancelledWhileTryingToConnect)
 			{
-				// run the completion delegate if we have one bound
-				ECommandResult::Type Result = ECommandResult::Failed;
-				if (Command.bCancelled)
-				{
-					Result = ECommandResult::Cancelled;
-				}
-				else if (Command.bCommandSuccessful)
-				{
-					Result = ECommandResult::Succeeded;
-				}
-				Command.OperationCompleteDelegate.ExecuteIfBound(Command.Operation, Result);
+				Command.ReturnResults();
 			}
 
 			//commands that are left in the array during a tick need to be deleted
@@ -390,7 +388,7 @@ void FPerforceSourceControlProvider::Tick()
 				delete &Command;
 			}
 
-			// only do one command per tick loop, as we dont want concurrent modification 
+			// only do one command per tick loop, as we dont want concurrent modification
 			// of the command queue (which can happen in the completion delegate)
 			break;
 		}
@@ -404,9 +402,7 @@ void FPerforceSourceControlProvider::Tick()
 			// when its (still running) thread finally finishes
 			Command.bAutoDelete = true;
 
-			// run the completion delegate if we have one bound
-			Command.OperationCompleteDelegate.ExecuteIfBound(Command.Operation, ECommandResult::Cancelled);
-
+			Command.ReturnResults();
 			break;
 		}
 	}
@@ -477,7 +473,7 @@ TSharedPtr<IPerforceSourceControlWorker, ESPMode::ThreadSafe> FPerforceSourceCon
 	{
 		return Operation->Execute();
 	}
-		
+
 	return NULL;
 }
 
@@ -544,7 +540,8 @@ ECommandResult::Type FPerforceSourceControlProvider::IssueCommand(FPerforceSourc
 {
 	if ( !bSynchronous && GThreadPool != NULL )
 	{
-		// Queue this to our worker thread(s) for resolving
+		// Queue this to our worker thread(s) for resolving.
+		// When asynchronous, any callback gets called from Tick().
 		GThreadPool->AddQueuedWork(&InCommand);
 		CommandQueue.Add(&InCommand);
 		return ECommandResult::Succeeded;
@@ -557,11 +554,7 @@ ECommandResult::Type FPerforceSourceControlProvider::IssueCommand(FPerforceSourc
 
 		OutputCommandMessages(InCommand);
 
-		// Callback now if present. When asynchronous, this callback gets called from Tick().
-		ECommandResult::Type Result = InCommand.bCommandSuccessful ? ECommandResult::Succeeded : ECommandResult::Failed;
-		InCommand.OperationCompleteDelegate.ExecuteIfBound(InCommand.Operation, Result);
-
-		return Result;
+		return InCommand.ReturnResults();
 	}
 }
 
diff --git a/Engine/Plugins/Developer/PerforceSourceControl/Source/PerforceSourceControl/Private/PerforceSourceControlSettings.cpp b/Engine/Plugins/Developer/PerforceSourceControl/Source/PerforceSourceControl/Private/PerforceSourceControlSettings.cpp
--- a/Engine/Plugins/Developer/PerforceSourceControl/Source/PerforceSourceControl/Private/PerforceSourceControlSettings.cpp
+++ b/Engine/Plugins/Developer/PerforceSourceControl/Source/PerforceSourceControl/Private/PerforceSourceControlSettings.cpp
@@ -6,6 +6,7 @@
 #include "Modules/ModuleManager.h"
 #include "SPerforceSourceControlSettings.h"
 #include "PerforceSourceControlModule.h"
+#include "SourceControlHelpers.h"
 
 namespace PerforceSettingsConstants
 {
diff --git a/Engine/Plugins/Developer/PerforceSourceControl/Source/PerforceSourceControl/Private/SPerforceSourceControlSettings.h b/Engine/Plugins/Developer/PerforceSourceControl/Source/PerforceSourceControl/Private/SPerforceSourceControlSettings.h
--- a/Engine/Plugins/Developer/PerforceSourceControl/Source/PerforceSourceControl/Private/SPerforceSourceControlSettings.h
+++ b/Engine/Plugins/Developer/PerforceSourceControl/Source/PerforceSourceControl/Private/SPerforceSourceControlSettings.h
@@ -11,7 +11,7 @@
 #include "Widgets/Input/SEditableTextBox.h"
 #include "Widgets/Input/SComboButton.h"
 #include "Widgets/Views/STableViewBase.h"
-#include "ISourceControlOperation.h"
+#include "SourceControlOperationBase.h"
 #include "ISourceControlProvider.h"
 
 class SButton;
@@ -19,13 +19,13 @@ class SButton;
 /**
  * Internal-only source control operation for retrieving available workspaces.
  */
-class FGetWorkspaces : public ISourceControlOperation
+class FGetWorkspaces : public FSourceControlOperationBase
 {
 public:
 	// ISourceControlOperation interface
-	virtual FName GetName() const override 
-	{ 
-		return "GetWorkspaces"; 
+	virtual FName GetName() const override
+	{
+		return "GetWorkspaces";
 	}
 
 public:
@@ -47,9 +47,9 @@ namespace ESourceControlOperationState
 class SPerforceSourceControlSettings : public SCompoundWidget
 {
 public:
-	
+
 	SLATE_BEGIN_ARGS(SPerforceSourceControlSettings) {}
-	
+
 	SLATE_END_ARGS()
 
 public:
@@ -58,7 +58,7 @@ public:
 
 	/** Get the currently entered password */
 	static FString GetPassword();
-	
+
 	virtual void Tick( const FGeometry& AllottedGeometry, const double InCurrentTime, const float InDeltaTime );
 private:
 
diff --git a/Engine/Plugins/Enterprise/DatasmithContent/Source/DatasmithContent/DatasmithContent.Build.cs b/Engine/Plugins/Enterprise/DatasmithContent/Source/DatasmithContent/DatasmithContent.Build.cs
--- a/Engine/Plugins/Enterprise/DatasmithContent/Source/DatasmithContent/DatasmithContent.Build.cs
+++ b/Engine/Plugins/Enterprise/DatasmithContent/Source/DatasmithContent/DatasmithContent.Build.cs
@@ -10,6 +10,7 @@ namespace UnrealBuildTool.Rules
 			PublicDependencyModuleNames.AddRange(
 				new string[]
 				{
+					"CinematicCamera",
 					"Core",
 					"CoreUObject",
 					"Engine"
diff --git a/Engine/Plugins/Enterprise/DatasmithContent/Source/DatasmithContent/Private/DatasmithAssetImportData.cpp b/Engine/Plugins/Enterprise/DatasmithContent/Source/DatasmithContent/Private/DatasmithAssetImportData.cpp
new file mode 100644
--- /dev/null
+++ b/Engine/Plugins/Enterprise/DatasmithContent/Source/DatasmithContent/Private/DatasmithAssetImportData.cpp
@@ -0,0 +1,161 @@
+// Copyright 1998-2018 Epic Games, Inc. All Rights Reserved.
+
+#include "DatasmithAssetImportData.h"
+#include "DatasmithContentModule.h"
+
+#include "Engine/StaticMesh.h"
+
+#include "HAL/FileManager.h"
+#include "Misc/FileHelper.h"
+#include "Misc/Paths.h"
+
+#if WITH_EDITORONLY_DATA
+UDatasmithStaticMeshImportData* UDatasmithStaticMeshImportData::GetImportDataForStaticMesh( UStaticMesh* StaticMesh, TOptional< DefaultOptionsPair > DefaultImportOptions )
+{
+	check( StaticMesh );
+
+	UDatasmithStaticMeshImportData* ImportData = nullptr;
+
+	ImportData = Cast< UDatasmithStaticMeshImportData >( StaticMesh->AssetImportData );
+	if ( !ImportData )
+	{
+		ImportData = NewObject< UDatasmithStaticMeshImportData >( StaticMesh );
+
+		if ( DefaultImportOptions.IsSet() )
+		{
+			ImportData->ImportOptions = DefaultImportOptions.GetValue().Key;
+			ImportData->AssetImportOptions = DefaultImportOptions.GetValue().Value;
+		}
+
+		// Try to preserve the source file path if possible
+		if ( StaticMesh->AssetImportData != nullptr )
+		{
+			ImportData->SourceData = StaticMesh->AssetImportData->SourceData;
+		}
+
+		StaticMesh->AssetImportData = ImportData;
+	}
+
+	return ImportData;
+}
+
+UDatasmithStaticMeshCADImportData * UDatasmithStaticMeshCADImportData::GetCADImportDataForStaticMesh(UStaticMesh * StaticMesh, TOptional< DefaultOptionsTuple > DefaultImportOptions)
+{
+	check( StaticMesh );
+
+	UDatasmithStaticMeshCADImportData* ImportData = Cast< UDatasmithStaticMeshCADImportData >( StaticMesh->AssetImportData );
+
+	if ( !ImportData )
+	{
+		ImportData = NewObject< UDatasmithStaticMeshCADImportData >( StaticMesh );
+
+
+		// Try to preserve the source file path if possible
+		if ( StaticMesh->AssetImportData != nullptr )
+		{
+			ImportData->SourceData = StaticMesh->AssetImportData->SourceData;
+		}
+
+		StaticMesh->AssetImportData = ImportData;
+	}
+
+	if ( DefaultImportOptions.IsSet() )
+	{
+		DefaultOptionsTuple& Value = DefaultImportOptions.GetValue();
+
+		ImportData->TessellationOptions = Value.Get<0>();
+		ImportData->ImportOptions = Value.Get<1>();
+		ImportData->AssetImportOptions = Value.Get<2>();
+	}
+
+	return ImportData;
+}
+
+void UDatasmithStaticMeshCADImportData::SetResourcePath(const FString& FilePath)
+{
+	check(FPaths::FileExists(FilePath));
+
+	ResourceFilename = FPaths::GetCleanFilename(FilePath);
+
+	TArray<FString> AuxililaryFilePaths;
+	IFileManager::Get().FindFiles(AuxililaryFilePaths, *(FilePath + TEXT(".*")), true, false);
+
+	AuxiliaryFilenames.Empty(AuxililaryFilePaths.Num());
+
+	for (FString& AuxililaryFilePath : AuxililaryFilePaths)
+	{
+		AuxiliaryFilenames.Add(FPaths::GetCleanFilename(AuxililaryFilePath));
+	}
+
+	// Set ResourcePath as absolute path because CoreTech expects an absolute path
+	ResourcePath = FPaths::ConvertRelativePathToFull(FilePath);
+}
+
+const FString& UDatasmithStaticMeshCADImportData::GetResourcePath()
+{
+	check(!ResourcePath.IsEmpty());
+
+	return ResourcePath;
+}
+
+void UDatasmithStaticMeshCADImportData::Serialize(FArchive& Ar)
+{
+	Super::Serialize(Ar);
+
+	// Do not bother loading or saving the rest if there is no source file
+	if (SourceData.SourceFiles.Num() == 0)
+	{
+		return;
+	}
+
+	FString ResourceDirectory;
+	TArray<uint8> ByteArray;
+	int32 AuxiliaryCount = AuxiliaryFilenames.Num();
+
+	if (Ar.IsSaving())
+	{
+		ResourceDirectory = FPaths::GetPath(ResourcePath);
+		FFileHelper::LoadFileToArray(ByteArray, *ResourcePath);
+	}
+	else if (Ar.IsLoading())
+	{
+		FString BaseName = FPaths::GetBaseFilename(SourceData.SourceFiles[0].RelativeFilename);
+		ResourceDirectory = FPaths::Combine(IDatasmithContentModule::Get().GetTempDir(), TEXT("ReimportCache"), *BaseName);
+		if (!FPaths::DirectoryExists(ResourceDirectory))
+		{
+			IFileManager::Get().MakeDirectory(*ResourceDirectory);
+		}
+	}
+
+	Ar << ByteArray;
+
+	if (Ar.IsSaving())
+	{
+		// Save each auxiliary file content
+		for (FString& Filename : AuxiliaryFilenames)
+		{
+			FString FilePath = FPaths::Combine(ResourceDirectory, Filename);
+
+			ByteArray.Empty();
+			FFileHelper::LoadFileToArray(ByteArray, *FilePath);
+
+			Ar << ByteArray;
+		}
+	}
+	else if (Ar.IsLoading())
+	{
+		ResourcePath = FPaths::ConvertRelativePathToFull(FPaths::Combine(ResourceDirectory, ResourceFilename));
+		FFileHelper::SaveArrayToFile(ByteArray, *ResourcePath);
+
+		// Save each auxiliary file content
+		for (FString& Filename : AuxiliaryFilenames)
+		{
+			ByteArray.Empty();
+			Ar << ByteArray;
+
+			FString FilePath = FPaths::Combine(ResourceDirectory, Filename);
+			FFileHelper::SaveArrayToFile(ByteArray, *FilePath);
+		}
+	}
+}
+#endif // WITH_EDITORONLY_DATA
diff --git a/Engine/Plugins/Enterprise/DatasmithContent/Source/DatasmithContent/Private/DatasmithContentModule.cpp b/Engine/Plugins/Enterprise/DatasmithContent/Source/DatasmithContent/Private/DatasmithContentModule.cpp
--- a/Engine/Plugins/Enterprise/DatasmithContent/Source/DatasmithContent/Private/DatasmithContentModule.cpp
+++ b/Engine/Plugins/Enterprise/DatasmithContent/Source/DatasmithContent/Private/DatasmithContentModule.cpp
@@ -3,6 +3,10 @@
 #include "DatasmithContentModule.h"
 #include "Modules/ModuleManager.h"
 
+#include "HAL/FileManager.h"
+#include "HAL/PlatformProcess.h"
+#include "Misc/Paths.h"
+
 /**
  * DatasmithContent module implementation (private)
  */
@@ -11,11 +15,23 @@ class FDatasmithContentModule : public IDatasmithContentModule
 public:
 	virtual void StartupModule() override
 	{
+		// Create temporary directory which will be used by UDatasmithStaticMeshCADImportData to store transient data
+		TempDir = FPaths::Combine(FPaths::ProjectIntermediateDir(), TEXT("DatasmithContentTemp"));
+		IFileManager::Get().MakeDirectory(*TempDir);
 	}
 
 	virtual void ShutdownModule() override
 	{
+		// Clean up all transient files created during the process
+		IFileManager::Get().DeleteDirectory(*TempDir, false, true);
+	}
+
+	const FString& GetTempDir() const override
+	{
+		return TempDir;
 	}
+private:
+	FString TempDir;
 };
 
 IMPLEMENT_MODULE(FDatasmithContentModule, DatasmithContent);
diff --git a/Engine/Plugins/Enterprise/DatasmithContent/Source/DatasmithContent/Private/ObjectTemplates/DatasmithSceneComponentTemplate.cpp b/Engine/Plugins/Enterprise/DatasmithContent/Source/DatasmithContent/Private/ObjectTemplates/DatasmithSceneComponentTemplate.cpp
new file mode 100644
--- /dev/null
+++ b/Engine/Plugins/Enterprise/DatasmithContent/Source/DatasmithContent/Private/ObjectTemplates/DatasmithSceneComponentTemplate.cpp
@@ -0,0 +1,53 @@
+// Copyright 1998-2018 Epic Games, Inc. All Rights Reserved.
+
+#include "ObjectTemplates/DatasmithSceneComponentTemplate.h"
+
+#include "DatasmithAssetUserData.h"
+#include "Components/SceneComponent.h"
+
+void UDatasmithSceneComponentTemplate::Apply( UObject* Destination, bool bForce )
+{
+#if WITH_EDITORONLY_DATA
+	USceneComponent* SceneComponent = Cast< USceneComponent >( Destination );
+
+	if ( !SceneComponent )
+	{
+		return;
+	}
+
+	UDatasmithSceneComponentTemplate* PreviousTemplate = !bForce ? FDatasmithObjectTemplateUtils::GetObjectTemplate< UDatasmithSceneComponentTemplate >( Destination ) : nullptr;
+
+	if ( !PreviousTemplate || PreviousTemplate->RelativeTransform.Equals( SceneComponent->GetRelativeTransform() ) )
+	{
+		SceneComponent->SetRelativeTransform( RelativeTransform );
+	}
+
+	if ( !PreviousTemplate || PreviousTemplate->Mobility == SceneComponent->Mobility )
+	{
+		SceneComponent->SetMobility( Mobility );
+	}
+
+	if ( !PreviousTemplate || PreviousTemplate->AttachParent == SceneComponent->GetAttachParent() )
+	{
+		SceneComponent->AttachToComponent( AttachParent, FAttachmentTransformRules::KeepRelativeTransform );
+	}
+
+	FDatasmithObjectTemplateUtils::SetObjectTemplate( Destination, this );
+#endif // #if WITH_EDITORONLY_DATA
+}
+
+void UDatasmithSceneComponentTemplate::Load( const UObject* Source )
+{
+#if WITH_EDITORONLY_DATA
+	const USceneComponent* SceneComponent = Cast< USceneComponent >( Source );
+
+	if ( !SceneComponent )
+	{
+		return;
+	}
+
+	RelativeTransform = SceneComponent->GetRelativeTransform();
+	Mobility = SceneComponent->Mobility;
+	AttachParent = SceneComponent->GetAttachParent();
+#endif // #if WITH_EDITORONLY_DATA
+}
\ No newline at end of file
diff --git a/Engine/Plugins/Enterprise/DatasmithContent/Source/DatasmithContent/Private/ObjectTemplates/DatasmithStaticMeshComponentTemplate.cpp b/Engine/Plugins/Enterprise/DatasmithContent/Source/DatasmithContent/Private/ObjectTemplates/DatasmithStaticMeshComponentTemplate.cpp
new file mode 100644
--- /dev/null
+++ b/Engine/Plugins/Enterprise/DatasmithContent/Source/DatasmithContent/Private/ObjectTemplates/DatasmithStaticMeshComponentTemplate.cpp
@@ -0,0 +1,67 @@
+// Copyright 1998-2018 Epic Games, Inc. All Rights Reserved.
+
+#include "ObjectTemplates/DatasmithStaticMeshComponentTemplate.h"
+
+#include "Components/StaticMeshComponent.h"
+
+void UDatasmithStaticMeshComponentTemplate::Apply( UObject* Destination, bool bForce )
+{
+#if WITH_EDITORONLY_DATA
+	UStaticMeshComponent* StaticMeshComponent = Cast< UStaticMeshComponent >( Destination );
+
+	if ( !StaticMeshComponent )
+	{
+		return;
+	}
+
+	UDatasmithStaticMeshComponentTemplate* PreviousStaticMeshTemplate = !bForce ? FDatasmithObjectTemplateUtils::GetObjectTemplate< UDatasmithStaticMeshComponentTemplate >( Destination ) : nullptr;
+
+	if ( !PreviousStaticMeshTemplate || PreviousStaticMeshTemplate->StaticMesh == StaticMeshComponent->GetStaticMesh() )
+	{
+		StaticMeshComponent->SetStaticMesh( StaticMesh );
+	}
+	
+	for ( int32 MaterialIndex = 0; MaterialIndex < OverrideMaterials.Num(); ++MaterialIndex )
+	{
+		// If it's a new material override, assign it
+		if ( !PreviousStaticMeshTemplate || !PreviousStaticMeshTemplate->OverrideMaterials.IsValidIndex( MaterialIndex ) )
+		{
+			StaticMeshComponent->SetMaterial( MaterialIndex, OverrideMaterials[ MaterialIndex ] );
+		}
+		else if ( StaticMeshComponent->OverrideMaterials.IsValidIndex( MaterialIndex ) &&
+			PreviousStaticMeshTemplate->OverrideMaterials[ MaterialIndex ] == StaticMeshComponent->OverrideMaterials[ MaterialIndex ] )
+		{
+			StaticMeshComponent->SetMaterial( MaterialIndex, OverrideMaterials[ MaterialIndex ] );
+		}
+	}
+
+	// Remove materials that aren't in the template anymore
+	if ( PreviousStaticMeshTemplate )
+	{
+		for ( int32 MaterialIndexToRemove = PreviousStaticMeshTemplate->OverrideMaterials.Num() - 1; MaterialIndexToRemove >= OverrideMaterials.Num(); --MaterialIndexToRemove )
+		{
+			if ( StaticMeshComponent->OverrideMaterials.IsValidIndex( MaterialIndexToRemove ) )
+			{
+				StaticMeshComponent->OverrideMaterials.RemoveAt( MaterialIndexToRemove );
+			}
+		}
+	}
+
+	FDatasmithObjectTemplateUtils::SetObjectTemplate( Destination, this );
+#endif // #if WITH_EDITORONLY_DATA
+}
+
+void UDatasmithStaticMeshComponentTemplate::Load( const UObject* Source )
+{
+#if WITH_EDITORONLY_DATA
+	const UStaticMeshComponent* SourceComponent = Cast< UStaticMeshComponent >( Source );
+
+	if ( !SourceComponent )
+	{
+		return;
+	}
+
+	StaticMesh = SourceComponent->GetStaticMesh();
+	OverrideMaterials = SourceComponent->OverrideMaterials;
+#endif // #if WITH_EDITORONLY_DATA
+}
\ No newline at end of file
diff --git a/Engine/Plugins/Enterprise/DatasmithContent/Source/DatasmithContent/Public/DatasmithAssetImportData.h b/Engine/Plugins/Enterprise/DatasmithContent/Source/DatasmithContent/Public/DatasmithAssetImportData.h
new file mode 100644
--- /dev/null
+++ b/Engine/Plugins/Enterprise/DatasmithContent/Source/DatasmithContent/Public/DatasmithAssetImportData.h
@@ -0,0 +1,110 @@
+// Copyright 1998-2018 Epic Games, Inc. All Rights Reserved.
+
+#pragma once
+
+#include "DatasmithImportOptions.h"
+
+#include "EditorFramework/AssetImportData.h"
+
+#include "DatasmithAssetImportData.generated.h"
+
+UCLASS()
+class DATASMITHCONTENT_API UDatasmithAssetImportData : public UAssetImportData
+{
+	GENERATED_BODY()
+
+#if WITH_EDITORONLY_DATA
+public:
+	UPROPERTY(BlueprintReadWrite, Category = Asset, meta = (ShowOnlyInnerProperties))
+	FDatasmithAssetImportOptions AssetImportOptions;
+#endif		// WITH_EDITORONLY_DATA
+};
+
+UCLASS()
+class DATASMITHCONTENT_API UDatasmithStaticMeshImportData : public UDatasmithAssetImportData
+{
+	GENERATED_BODY()
+
+#if WITH_EDITORONLY_DATA
+public:
+	UPROPERTY(EditAnywhere, Category = StaticMesh, meta = (ShowOnlyInnerProperties))
+	FDatasmithStaticMeshImportOptions ImportOptions;
+
+public:
+	typedef TPair< FDatasmithStaticMeshImportOptions, FDatasmithAssetImportOptions > DefaultOptionsPair;
+	static UDatasmithStaticMeshImportData* GetImportDataForStaticMesh( UStaticMesh* StaticMesh, TOptional< DefaultOptionsPair > DefaultImportOptions = TOptional< DefaultOptionsPair >() );
+
+#endif		// WITH_EDITORONLY_DATA
+};
+
+UCLASS(HideCategories = (InternalProperty))
+class DATASMITHCONTENT_API UDatasmithStaticMeshCADImportData : public UDatasmithStaticMeshImportData
+{
+	GENERATED_BODY()
+
+#if WITH_EDITORONLY_DATA
+public:
+	UPROPERTY(EditAnywhere, Category = StaticMesh, meta = (ShowOnlyInnerProperties))
+	FDatasmithTessellationOptions TessellationOptions;
+
+	UPROPERTY(VisibleDefaultsOnly, Category = InternalProperty)
+	double	ModelUnit;
+
+	UPROPERTY(VisibleDefaultsOnly, Category = InternalProperty)
+	double	ModelTolerance;
+
+protected:
+	UPROPERTY(VisibleDefaultsOnly, Category = InternalProperty)
+	FString	ResourcePath;
+
+	UPROPERTY(VisibleDefaultsOnly, Category = InternalProperty)
+	FString	ResourceFilename;
+
+	UPROPERTY(VisibleDefaultsOnly, Category = InternalProperty)
+	TArray<FString>	AuxiliaryFilenames;
+
+public:
+	typedef TTuple<FDatasmithTessellationOptions, FDatasmithStaticMeshImportOptions, FDatasmithAssetImportOptions> DefaultOptionsTuple;
+	static UDatasmithStaticMeshCADImportData* GetCADImportDataForStaticMesh( UStaticMesh* StaticMesh, TOptional< DefaultOptionsTuple > DefaultImportOptions = TOptional< DefaultOptionsTuple >() );
+
+	void SetResourcePath(const FString& FilePath);
+	const FString& GetResourcePath();
+
+protected:
+	/** Overridden serialize function to read in/write out the unexposed data */
+	virtual void Serialize(FArchive& Ar) override;
+#endif		// WITH_EDITORONLY_DATA
+};
+
+/**
+ * Base class for import data and options used when importing any asset from Datasmith
+ */
+UCLASS()
+class DATASMITHCONTENT_API UDatasmithSceneImportData : public UAssetImportData
+{
+	GENERATED_BODY()
+
+#if WITH_EDITORONLY_DATA
+public:
+	UPROPERTY(EditAnywhere, Category = "Options", meta = (ShowOnlyInnerProperties))
+	FDatasmithImportBaseOptions BaseOptions;
+
+	UPROPERTY()
+	FString SceneName;
+#endif		// WITH_EDITORONLY_DATA
+};
+
+/**
+ * Import data and options specific to tessellated Datasmith scenes
+ */
+UCLASS()
+class DATASMITHCONTENT_API UDatasmithCADImportSceneData : public UDatasmithSceneImportData
+{
+	GENERATED_BODY()
+
+#if WITH_EDITORONLY_DATA
+public:
+	UPROPERTY(EditAnywhere, Category = "Tessellation", meta = (ShowOnlyInnerProperties))
+	FDatasmithTessellationOptions TessellationOptions;
+#endif		// WITH_EDITORONLY_DATA
+};
\ No newline at end of file
diff --git a/Engine/Plugins/Enterprise/DatasmithContent/Source/DatasmithContent/Public/DatasmithAssetUserData.h b/Engine/Plugins/Enterprise/DatasmithContent/Source/DatasmithContent/Public/DatasmithAssetUserData.h
new file mode 100644
--- /dev/null
+++ b/Engine/Plugins/Enterprise/DatasmithContent/Source/DatasmithContent/Public/DatasmithAssetUserData.h
@@ -0,0 +1,28 @@
+// Copyright 1998-2018 Epic Games, Inc. All Rights Reserved.
+
+#pragma once
+
+#include "Engine/AssetUserData.h"
+#include "Templates/SubclassOf.h"
+
+#include "DatasmithAssetUserData.generated.h"
+
+/** Asset user data that can be used with Datasmith on Actors and other objects  */
+UCLASS(BlueprintType, meta = (ScriptName = "DatasmithUserData", DisplayName = "Datasmith User Data"))
+class DATASMITHCONTENT_API UDatasmithAssetUserData : public UAssetUserData
+{
+	GENERATED_BODY()
+
+public:
+
+	// Meta-data are available at runtime in game, i.e. used in blueprint to display build-boarded information
+	typedef TMap<FName, FString> FMetaDataContainer;
+
+	UPROPERTY(VisibleAnywhere, BlueprintReadWrite, AdvancedDisplay, Category = "Datasmith User Data", meta = (ScriptName = "Metadata", DisplayName = "Metadata"))
+	TMap<FName, FString> MetaData;
+
+#if WITH_EDITORONLY_DATA
+	UPROPERTY(VisibleAnywhere, Category = "Datasmith User Data")
+	TMap< TSubclassOf< class UDatasmithObjectTemplate >, UDatasmithObjectTemplate* > ObjectTemplates;
+#endif
+};
diff --git a/Engine/Plugins/Enterprise/DatasmithContent/Source/DatasmithContent/Public/DatasmithContentModule.h b/Engine/Plugins/Enterprise/DatasmithContent/Source/DatasmithContent/Public/DatasmithContentModule.h
--- a/Engine/Plugins/Enterprise/DatasmithContent/Source/DatasmithContent/Public/DatasmithContentModule.h
+++ b/Engine/Plugins/Enterprise/DatasmithContent/Source/DatasmithContent/Public/DatasmithContentModule.h
@@ -34,5 +34,10 @@ public:
 	{
 		return FModuleManager::Get().IsModuleLoaded(DATASMITHCONTENT_MODULE_NAME);
 	}
+
+	/**
+	 * @returns temporary directory created for this process
+	 */
+	virtual const FString& GetTempDir() const = 0;
 };
 
diff --git a/Engine/Plugins/Enterprise/DatasmithContent/Source/DatasmithContent/Public/DatasmithImportOptions.h b/Engine/Plugins/Enterprise/DatasmithContent/Source/DatasmithContent/Public/DatasmithImportOptions.h
new file mode 100644
--- /dev/null
+++ b/Engine/Plugins/Enterprise/DatasmithContent/Source/DatasmithContent/Public/DatasmithImportOptions.h
@@ -0,0 +1,257 @@
+// Copyright 1998-2018 Epic Games, Inc. All Rights Reserved.
+
+#pragma once
+
+#include "CoreMinimal.h"
+#include "UObject/ObjectMacros.h"
+#include "Engine/EngineTypes.h"
+
+#include "DatasmithImportOptions.generated.h"
+
+class FJsonObject;
+
+UENUM()
+enum class EDatasmithImportSearchPackagePolicy : uint8
+{
+	/** Search only in current package */
+	Current UMETA(DisplayName = "Current", DisplayValue = "Current"),
+
+	/** Search in all packages */
+	All,
+};
+
+UENUM()
+enum class EDatasmithImportAssetConflictPolicy : uint8
+{
+	/** Replace existing asset with new one */
+	Replace,
+
+	/** Update existing asset with new values */
+	Update,
+
+	/** Use existing asset instead of creating new one */
+	Use,
+
+	/** Skip new asset */
+	Ignore,
+};
+
+UENUM()
+enum class EDatasmithImportMaterialQuality : uint8
+{
+	UseNoFresnelCurves,
+
+	UseSimplifierFresnelCurves,
+
+	UseRealFresnelCurves,
+};
+
+UENUM()
+enum class EDatasmithImportLightmapMin : uint8
+{
+	LIGHTMAP_16		UMETA(DisplayName = "16"),
+	LIGHTMAP_32		UMETA(DisplayName = "32"),
+	LIGHTMAP_64		UMETA(DisplayName = "64"),
+	LIGHTMAP_128	UMETA(DisplayName = "128"),
+	LIGHTMAP_256	UMETA(DisplayName = "256"),
+	LIGHTMAP_512	UMETA(DisplayName = "512"),
+};
+
+UENUM()
+enum class EDatasmithImportLightmapMax : uint8
+{
+	LIGHTMAP_64		UMETA(DisplayName = "64"),
+	LIGHTMAP_128	UMETA(DisplayName = "128"),
+	LIGHTMAP_256	UMETA(DisplayName = "256"),
+	LIGHTMAP_512	UMETA(DisplayName = "512"),
+	LIGHTMAP_1024	UMETA(DisplayName = "1024"),
+	LIGHTMAP_2048	UMETA(DisplayName = "2048"),
+	LIGHTMAP_4096	UMETA(DisplayName = "4096")
+};
+
+UENUM()
+enum class EDatasmithImportScene : uint8
+{
+	NewLevel		UMETA(DisplayName = "Create New Level", ToolTip = "Create a new Level and spawn the actors after the import."),
+
+	CurrentLevel	UMETA(DisplayName = "Merge to Current Level", ToolTip = "Use the current Level to spawn the actors after the import."),
+
+	AssetsOnly		UMETA(DisplayName = "Assets Only", ToolTip = "Do not modify the Level after import. No actor will be created (including the Blueprint if requested by the ImportHierarchy"),
+};
+
+UENUM()
+enum class EDatasmithImportHierarchy : uint8
+{
+	UseMultipleActors	UMETA(DisplayName = "One StaticMeshActor per Geometric Object", ToolTip = "Create an StaticMeshActor for every node in the hierarchy of the model."),
+
+	UseSingleActor		UMETA(DisplayName = "Single StaticMeshActor with Components", ToolTip = "Create one root StaticMeshActor then one component for every node in the hierarchy of the model. Recommended to import udatasmith files."),
+
+	UseOneBlueprint		UMETA(DisplayName = "Blueprint", ToolTip = "Create one root blueprint then one component for every node in the hierarchy of the model. Recommended to import CAD files."),
+};
+
+USTRUCT(BlueprintType)
+struct DATASMITHCONTENT_API FDatasmithAssetImportOptions
+{
+	GENERATED_USTRUCT_BODY()
+
+	UPROPERTY()
+	FName PackagePath;
+};
+
+USTRUCT(BlueprintType)
+struct DATASMITHCONTENT_API FDatasmithStaticMeshImportOptions
+{
+	GENERATED_USTRUCT_BODY()
+
+	FDatasmithStaticMeshImportOptions();
+
+	UPROPERTY(EditAnywhere, BlueprintReadWrite, Category = Lightmap)
+	EDatasmithImportLightmapMin MinLightmapResolution;
+
+	UPROPERTY(EditAnywhere, BlueprintReadWrite, Category = Lightmap)
+	EDatasmithImportLightmapMax MaxLightmapResolution;
+
+	UPROPERTY(BlueprintReadWrite, Category = Mesh)
+	bool bRemoveDegenerates;
+
+public:
+	static int32 ConvertLightmapEnumToValue( EDatasmithImportLightmapMin EnumValue );
+	static int32 ConvertLightmapEnumToValue( EDatasmithImportLightmapMax EnumValue );
+
+	bool operator == (const FDatasmithStaticMeshImportOptions& Other) const
+	{
+		return (MinLightmapResolution == Other.MinLightmapResolution && MaxLightmapResolution == Other.MaxLightmapResolution && bRemoveDegenerates == Other.bRemoveDegenerates);
+	}
+};
+
+USTRUCT(BlueprintType)
+struct DATASMITHCONTENT_API FDatasmithImportBaseOptions
+{
+	GENERATED_USTRUCT_BODY()
+
+	FDatasmithImportBaseOptions();
+
+	/** Specifies where to put the content */
+	UPROPERTY(BlueprintReadWrite, Category = Import)
+	EDatasmithImportScene SceneHandling; // Not displayed
+
+	/** Specifies whether geometry are to be imported or not */
+	UPROPERTY(config, EditAnywhere, BlueprintReadWrite, Category = Include, meta = (DisplayName = "Geometry"))
+	bool bIncludeGeometry;
+
+	/** Specifies whether materials and textures are to be imported or not */
+	UPROPERTY(config, EditAnywhere, BlueprintReadWrite, Category = Include, meta = (DisplayName = "Materials & Textures"))
+	bool bIncludeMaterial;
+
+	/** Specifies whether lights are to be imported or not */
+	UPROPERTY(config, EditAnywhere, BlueprintReadWrite, Category = Include, meta = (DisplayName = "Lights"))
+	bool bIncludeLight;
+
+	/** Specifies whether cameras are to be imported or not */
+	UPROPERTY(config, EditAnywhere, BlueprintReadWrite, Category = Include, meta = (DisplayName = "Cameras"))
+	bool bIncludeCamera;
+
+	UPROPERTY(BlueprintReadWrite, AdvancedDisplay, Category = Include, meta = (ShowOnlyInnerProperties))
+	FDatasmithAssetImportOptions AssetOptions;
+
+	UPROPERTY(config, EditAnywhere, BlueprintReadWrite, AdvancedDisplay, Category = Include, meta = (ShowOnlyInnerProperties))
+	FDatasmithStaticMeshImportOptions StaticMeshOptions;
+};
+
+USTRUCT(BlueprintType)
+struct DATASMITHCONTENT_API FDatasmithTessellationOptions
+{
+	GENERATED_USTRUCT_BODY()
+
+	FDatasmithTessellationOptions()
+		: ChordTolerance(0.2f)
+		, MaxEdgeLength(0.0f)
+		, NormalTolerance(20.0f)
+	{
+
+	}
+
+	/**
+	 * Maximum distance between any point on a triangle generated by the tessellation process and the actual surface.
+	 * The lower the value the more triangles.
+	 * Default value is 0.2.
+	 */
+	UPROPERTY(config, EditAnywhere, BlueprintReadWrite, Category = "Geometry & Tessellation Options", meta = (Units = cm, ToolTip = "Maximum distance between a generated triangle and the original surface. Smaller values increase triangles count.", ClampMin = "0.0"))
+	float ChordTolerance;
+
+	/**
+	 * Maximum length of edges of triangles generated by the tessellation process.
+	 * The length is in scene/model unit. The smaller the more triangles are generated.
+	 * Value of 0 means no constraint on length of edges
+	 * Default value is 0.
+	 */
+	UPROPERTY(config, EditAnywhere, BlueprintReadWrite, AdvancedDisplay, Category = "Geometry & Tessellation Options", meta = (Units = cm, DisplayName = "Max Edge Length", ToolTip = "Maximum length of an edge in the generated triangles. Smaller values increase triangles count.", ClampMin = "0.0"))
+	float MaxEdgeLength;
+
+	/**
+	 * Maximum angle between the normal of two triangles generated by the tessellation process.
+	 * The angle is expressed in degree. The smaller the more triangles are generated.
+	 * Default value is 20 degrees.
+	 */
+	UPROPERTY(config, EditAnywhere, BlueprintReadWrite, AdvancedDisplay, Category = "Geometry & Tessellation Options", meta = (Units = deg, ToolTip = "Maximum angle between adjacent triangles generated from a surface. Smaller values increase triangles count.", ClampMin = "0.0", ClampMax = "90.0"))
+	float NormalTolerance;
+
+public:
+	bool operator == (const FDatasmithTessellationOptions& Other) const
+	{
+		return (FMath::IsNearlyEqual(ChordTolerance, Other.ChordTolerance) && FMath::IsNearlyEqual(MaxEdgeLength, Other.MaxEdgeLength) && FMath::IsNearlyEqual(NormalTolerance, Other.NormalTolerance));
+	}
+};
+
+UCLASS(config = EditorPerProjectUserSettings, HideCategories = ("TessellationOff"))
+class DATASMITHCONTENT_API UDatasmithImportOptions : public UObject
+{
+	GENERATED_UCLASS_BODY()
+
+public:
+	/** Specifies where to search for assets */
+	UPROPERTY(Transient, AdvancedDisplay)
+	EDatasmithImportSearchPackagePolicy SearchPackagePolicy; // Not displayed. Kept for future use
+
+	/** Specifies what to do when material conflicts */
+	UPROPERTY(Transient, AdvancedDisplay)
+	EDatasmithImportAssetConflictPolicy MaterialConflictPolicy; // Not displayed. Kept for future use
+
+	/** Specifies what to do when texture conflicts */
+	UPROPERTY(Transient, AdvancedDisplay)
+	EDatasmithImportAssetConflictPolicy TextureConflictPolicy; // Not displayed. Kept for future use
+
+	/** Specifies what to do when actor conflicts */
+	UPROPERTY(Transient, AdvancedDisplay)
+	EDatasmithImportAssetConflictPolicy ActorConflictPolicy; // Not displayed. Kept for future use
+
+	/** Specifies what to do when light conflicts */
+	UPROPERTY(Transient, AdvancedDisplay)
+	EDatasmithImportAssetConflictPolicy LightConflictPolicy; // Not displayed. Kept for future use
+
+	/** Specifies what to do when material conflicts */
+	UPROPERTY(Transient, AdvancedDisplay)
+	EDatasmithImportAssetConflictPolicy CameraConflictPolicy; // Not displayed. Kept for future use
+
+	/** Specifies what to do when material conflicts */
+	UPROPERTY(Transient, AdvancedDisplay)
+	EDatasmithImportMaterialQuality MaterialQuality; // Not displayed. Kept for future use
+
+	/** Specifies how to import the model's hierarchy */
+	EDatasmithImportHierarchy HierarchyHandling;
+
+	UPROPERTY(config, EditAnywhere, BlueprintReadWrite, Category = "Options", meta = (ShowOnlyInnerProperties))
+	FDatasmithImportBaseOptions BaseOptions;
+
+	UPROPERTY(config, EditAnywhere, BlueprintReadWrite, Category = "TessellationOff", meta = (ShowOnlyInnerProperties))
+	FDatasmithTessellationOptions TessellationOptions;
+
+	/** Name of the imported file without its path */
+	FString FileName;
+	/** Full path of the imported file */
+	FString FilePath;
+	/** Whether to use or not the same options when loading multiple files. Default false */
+	bool bUseSameOptions;
+
+	void UpdateNotDisplayedConfig();
+};
diff --git a/Engine/Plugins/Enterprise/DatasmithContent/Source/DatasmithContent/Public/DatasmithSceneActor.h b/Engine/Plugins/Enterprise/DatasmithContent/Source/DatasmithContent/Public/DatasmithSceneActor.h
new file mode 100644
--- /dev/null
+++ b/Engine/Plugins/Enterprise/DatasmithContent/Source/DatasmithContent/Public/DatasmithSceneActor.h
@@ -0,0 +1,18 @@
+// Copyright 1998-2018 Epic Games, Inc. All Rights Reserved.
+
+#pragma once
+
+#include "GameFramework/Actor.h"
+
+#include "DatasmithSceneActor.generated.h"
+
+UCLASS()
+class DATASMITHCONTENT_API ADatasmithSceneActor : public AActor
+{
+	GENERATED_BODY()
+
+public:
+
+	UPROPERTY(VisibleAnywhere, Category="Datasmith")
+	class UDatasmithSceneImportData* Scene;
+};
diff --git a/Engine/Plugins/Enterprise/DatasmithContent/Source/DatasmithContent/Public/ObjectTemplates/DatasmithAreaLightActorTemplate.h b/Engine/Plugins/Enterprise/DatasmithContent/Source/DatasmithContent/Public/ObjectTemplates/DatasmithAreaLightActorTemplate.h
new file mode 100644
--- /dev/null
+++ b/Engine/Plugins/Enterprise/DatasmithContent/Source/DatasmithContent/Public/ObjectTemplates/DatasmithAreaLightActorTemplate.h
@@ -0,0 +1,33 @@
+// Copyright 1998-2018 Epic Games, Inc. All Rights Reserved.
+
+#pragma once
+
+#include "DatasmithAreaLightActor.h"
+#include "ObjectTemplates/DatasmithObjectTemplate.h"
+
+#include "DatasmithAreaLightActorTemplate.generated.h"
+
+UCLASS()
+class DATASMITHCONTENT_API UDatasmithAreaLightActorTemplate : public UDatasmithObjectTemplate
+{
+	GENERATED_BODY()
+
+public:
+	UPROPERTY()
+	EDatasmithAreaLightActorShape LightShape;
+
+	UPROPERTY()
+	FVector2D Dimensions;
+
+	UPROPERTY()
+	FLinearColor Color;
+
+	UPROPERTY()
+	float Intensity;
+
+	UPROPERTY()
+	uint8 bHidden:1;
+
+	virtual void Apply( UObject* Destination, bool bForce = false ) override;
+	virtual void Load( const UObject* Source ) override;
+};
\ No newline at end of file
diff --git a/Engine/Plugins/Enterprise/DatasmithContent/Source/DatasmithContent/Public/ObjectTemplates/DatasmithCineCameraComponentTemplate.h b/Engine/Plugins/Enterprise/DatasmithContent/Source/DatasmithContent/Public/ObjectTemplates/DatasmithCineCameraComponentTemplate.h
new file mode 100644
--- /dev/null
+++ b/Engine/Plugins/Enterprise/DatasmithContent/Source/DatasmithContent/Public/ObjectTemplates/DatasmithCineCameraComponentTemplate.h
@@ -0,0 +1,114 @@
+// Copyright 1998-2018 Epic Games, Inc. All Rights Reserved.
+
+#pragma once
+
+#include "ObjectTemplates/DatasmithObjectTemplate.h"
+
+#include "DatasmithCineCameraComponentTemplate.generated.h"
+
+struct FCameraFilmbackSettings;
+struct FCameraFocusSettings;
+struct FCameraLensSettings;
+
+USTRUCT()
+struct FDatasmithCameraFilmbackSettingsTemplate
+{
+	GENERATED_BODY()
+
+	UPROPERTY()
+	float SensorWidth;
+
+	UPROPERTY()
+	float SensorHeight;
+
+	void Apply( FCameraFilmbackSettings* Destination, const FDatasmithCameraFilmbackSettingsTemplate* PreviousTemplate );
+	void Load( const FCameraFilmbackSettings& Source );
+};
+
+USTRUCT()
+struct FDatasmithCameraLensSettingsTemplate
+{
+	GENERATED_BODY()
+
+	UPROPERTY()
+	float MaxFStop;
+
+	void Apply( FCameraLensSettings* Destination, const FDatasmithCameraLensSettingsTemplate* PreviousTemplate );
+	void Load( const FCameraLensSettings& Source );
+};
+
+USTRUCT()
+struct FDatasmithCameraFocusSettingsTemplate
+{
+	GENERATED_BODY()
+
+	UPROPERTY()
+	float ManualFocusDistance;
+
+	void Apply( FCameraFocusSettings* Destination, const FDatasmithCameraFocusSettingsTemplate* PreviousTemplate );
+	void Load( const FCameraFocusSettings& Source );
+};
+
+USTRUCT()
+struct DATASMITHCONTENT_API FDatasmithPostProcessSettingsTemplate
+{
+	GENERATED_BODY()
+
+public:
+	FDatasmithPostProcessSettingsTemplate();
+
+	UPROPERTY()
+	uint32 bOverride_WhiteTemp:1;
+
+	UPROPERTY()
+	uint32 bOverride_ColorSaturation:1;
+
+	UPROPERTY()
+	uint32 bOverride_VignetteIntensity:1;
+
+	UPROPERTY()
+	uint32 bOverride_FilmWhitePoint:1;
+
+	UPROPERTY()
+	float WhiteTemp;
+
+	UPROPERTY()
+	float VignetteIntensity;
+
+	UPROPERTY()
+	FLinearColor FilmWhitePoint;
+
+	UPROPERTY()
+	FVector4 ColorSaturation;
+
+	void Apply( struct FPostProcessSettings* Destination, const FDatasmithPostProcessSettingsTemplate* PreviousTemplate );
+	void Load( const FPostProcessSettings& Source );
+};
+
+UCLASS()
+class DATASMITHCONTENT_API UDatasmithCineCameraComponentTemplate : public UDatasmithObjectTemplate
+{
+	GENERATED_BODY()
+
+public:
+	UPROPERTY()
+	FDatasmithCameraFilmbackSettingsTemplate FilmbackSettings;
+
+	UPROPERTY()
+	FDatasmithCameraLensSettingsTemplate LensSettings;
+
+	UPROPERTY()
+	FDatasmithCameraFocusSettingsTemplate FocusSettings;
+
+	UPROPERTY()
+	float CurrentFocalLength;
+
+	UPROPERTY()
+	float CurrentAperture;
+
+	UPROPERTY()
+	FDatasmithPostProcessSettingsTemplate PostProcessSettings;
+
+	virtual void Apply( UObject* Destination, bool bForce = false ) override;
+	virtual void Load( const UObject* Source ) override;
+};
diff --git a/Engine/Plugins/Enterprise/DatasmithContent/Source/DatasmithContent/Public/ObjectTemplates/DatasmithLightComponentTemplate.h b/Engine/Plugins/Enterprise/DatasmithContent/Source/DatasmithContent/Public/ObjectTemplates/DatasmithLightComponentTemplate.h
new file mode 100644
--- /dev/null
+++ b/Engine/Plugins/Enterprise/DatasmithContent/Source/DatasmithContent/Public/ObjectTemplates/DatasmithLightComponentTemplate.h
@@ -0,0 +1,49 @@
+// Copyright 1998-2018 Epic Games, Inc. All Rights Reserved.
+
+#pragma once
+
+#include "ObjectTemplates/DatasmithObjectTemplate.h"
+
+#include "DatasmithLightComponentTemplate.generated.h"
+
+UCLASS()
+class DATASMITHCONTENT_API UDatasmithLightComponentTemplate : public UDatasmithObjectTemplate
+{
+	GENERATED_BODY()
+
+public:
+	UDatasmithLightComponentTemplate();
+
+	UPROPERTY()
+	uint8 bVisible : 1;
+
+	UPROPERTY()
+	uint32 CastShadows : 1;
+
+	UPROPERTY()
+	uint32 bUseTemperature : 1;
+
+	UPROPERTY()
+	uint32 bUseIESBrightness : 1;
+
+	UPROPERTY()
+	float Intensity;
+
+	UPROPERTY()
+	float Temperature;
+
+	UPROPERTY()
+	float IESBrightnessScale;
+
+	UPROPERTY()
+	FLinearColor LightColor;
+
+	UPROPERTY()
+	class UMaterialInterface* LightFunctionMaterial;
+
+	UPROPERTY()
+	class UTextureLightProfile* IESTexture;
+
+	virtual void Apply( UObject* Destination, bool bForce = false ) override;
+	virtual void Load( const UObject* Source ) override;
+};
diff --git a/Engine/Plugins/Enterprise/DatasmithContent/Source/DatasmithContent/Public/ObjectTemplates/DatasmithObjectTemplate.h b/Engine/Plugins/Enterprise/DatasmithContent/Source/DatasmithContent/Public/ObjectTemplates/DatasmithObjectTemplate.h
new file mode 100644
--- /dev/null
+++ b/Engine/Plugins/Enterprise/DatasmithContent/Source/DatasmithContent/Public/ObjectTemplates/DatasmithObjectTemplate.h
@@ -0,0 +1,136 @@
+// Copyright 1998-2018 Epic Games, Inc. All Rights Reserved.
+
+#pragma once
+
+#include "DatasmithAssetUserData.h"
+#include "Interfaces/Interface_AssetUserData.h"
+#include "UObject/Object.h"
+
+#include "DatasmithObjectTemplate.generated.h"
+
+UCLASS(abstract)
+class DATASMITHCONTENT_API UDatasmithObjectTemplate : public UObject
+{
+	GENERATED_BODY()
+
+public:
+	/**
+	 * Applies the object template to a Destination object
+	 *
+	 * @param Destination	The object to apply this template to
+	 * @param bForce		Force the application of the template on all properties, even if they were changed from the previous template values.
+	 */
+	virtual void Apply( UObject* Destination, bool bForce = false ) {}
+
+	/**
+	 * Fills this template properties with the values from the Source object.
+	 */
+	virtual void Load( const UObject* Source ) {}
+};
+
+// Sets Destination->MemberName with the value of MemberName only if PreviousTemplate is null or has the same value for MemberName as the Destination.
+// The goal is to set a new value only if it wasn't changed (overriden) in the Destination.
+#define DATASMITHOBJECTTEMPLATE_CONDITIONALSET(MemberName, Destination, PreviousTemplate) \
+	if ( !PreviousTemplate || Destination->MemberName == PreviousTemplate->MemberName ) \
+	{ \
+		Destination->MemberName = MemberName; \
+	}
+
+struct FDatasmithObjectTemplateUtils
+{
+	inline static bool HasObjectTemplates( UObject* Outer )
+	{
+#if WITH_EDITORONLY_DATA
+		if ( !Outer )
+		{
+			return false;
+		}
+
+		UDatasmithAssetUserData* UserData = nullptr;
+
+		if ( Outer->GetClass()->ImplementsInterface( UInterface_AssetUserData::StaticClass() ) )
+		{
+			IInterface_AssetUserData* AssetUserDataInterface = Cast< IInterface_AssetUserData >( Outer );
+
+			UserData = AssetUserDataInterface->GetAssetUserData< UDatasmithAssetUserData >();
+		}
+
+		return UserData != nullptr && UserData->ObjectTemplates.Num() > 0;
+#else
+		return false;
+#endif // #if WITH_EDITORONLY_DATA
+	}
+
+	inline static TMap< TSubclassOf< UDatasmithObjectTemplate >, UDatasmithObjectTemplate* >* FindOrCreateObjectTemplates( UObject* Outer )
+	{
+#if WITH_EDITORONLY_DATA
+		if ( !Outer )
+		{
+			return nullptr;
+		}
+
+		UDatasmithAssetUserData* UserData = nullptr;
+
+		if ( Outer->GetClass()->ImplementsInterface( UInterface_AssetUserData::StaticClass() ) )
+		{
+			IInterface_AssetUserData* AssetUserDataInterface = Cast< IInterface_AssetUserData >( Outer );
+
+			UserData = AssetUserDataInterface->GetAssetUserData< UDatasmithAssetUserData >();
+
+			if ( !UserData )
+			{
+				UserData = NewObject< UDatasmithAssetUserData >( Outer, NAME_None );
+				AssetUserDataInterface->AddAssetUserData( UserData );
+			}
+		}
+
+		if ( !UserData )
+		{
+			return nullptr;
+		}
+
+		return &UserData->ObjectTemplates;
+#else
+		return nullptr;
+#endif // #if WITH_EDITORONLY_DATA
+	}
+
+	inline static UDatasmithObjectTemplate* GetObjectTemplate( UObject* Outer, TSubclassOf< UDatasmithObjectTemplate > Subclass )
+	{
+#if WITH_EDITORONLY_DATA
+		TMap< TSubclassOf< UDatasmithObjectTemplate >, UDatasmithObjectTemplate* >* ObjectTemplatesMap = FindOrCreateObjectTemplates( Outer );
+
+		if ( !ObjectTemplatesMap )
+		{
+			return nullptr;
+		}
+
+		UDatasmithObjectTemplate** ObjectTemplatePtr = ObjectTemplatesMap->Find( Subclass );
+
+		return ObjectTemplatePtr ? *ObjectTemplatePtr : nullptr;
+#else
+		return nullptr;
+#endif // #if WITH_EDITORONLY_DATA
+	}
+
+	template < typename T >
+	static T* GetObjectTemplate( UObject* Outer )
+	{
+		return Cast< T >( GetObjectTemplate( Outer, T::StaticClass() ) );
+	}
+
+	inline static void SetObjectTemplate( UObject* Outer, UDatasmithObjectTemplate* ObjectTemplate )
+	{
+#if WITH_EDITORONLY_DATA
+		TMap< TSubclassOf< UDatasmithObjectTemplate >, UDatasmithObjectTemplate* >* ObjectTemplatesMap = FindOrCreateObjectTemplates( Outer );
+		ensure( ObjectTemplatesMap );
+
+		if ( !ObjectTemplatesMap )
+		{
+			return;
+		}
+
+		ObjectTemplatesMap->FindOrAdd( ObjectTemplate->GetClass() ) = ObjectTemplate;
+#endif // #if WITH_EDITORONLY_DATA
+	}
+};
diff --git a/Engine/Plugins/Enterprise/DatasmithContent/Source/DatasmithContent/Public/ObjectTemplates/DatasmithPointLightComponentTemplate.h b/Engine/Plugins/Enterprise/DatasmithContent/Source/DatasmithContent/Public/ObjectTemplates/DatasmithPointLightComponentTemplate.h
new file mode 100644
--- /dev/null
+++ b/Engine/Plugins/Enterprise/DatasmithContent/Source/DatasmithContent/Public/ObjectTemplates/DatasmithPointLightComponentTemplate.h
@@ -0,0 +1,32 @@
+// Copyright 1998-2018 Epic Games, Inc. All Rights Reserved.
+
+#pragma once
+
+#include "Engine/Scene.h"
+#include "ObjectTemplates/DatasmithObjectTemplate.h"
+
+#include "DatasmithPointLightComponentTemplate.generated.h"
+
+UCLASS()
+class DATASMITHCONTENT_API UDatasmithPointLightComponentTemplate : public UDatasmithObjectTemplate
+{
+	GENERATED_BODY()
+
+public:
+	UDatasmithPointLightComponentTemplate();
+
+	UPROPERTY()
+	ELightUnits IntensityUnits;
+
+	UPROPERTY()
+	float SourceRadius;
+
+	UPROPERTY()
+	float SourceLength;
+
+	UPROPERTY()
+	float AttenuationRadius;
+
+	virtual void Apply( UObject* Destination, bool bForce = false ) override;
+	virtual void Load( const UObject* Source ) override;
+};
diff --git a/Engine/Plugins/Enterprise/DatasmithContent/Source/DatasmithContent/Public/ObjectTemplates/DatasmithSceneComponentTemplate.h b/Engine/Plugins/Enterprise/DatasmithContent/Source/DatasmithContent/Public/ObjectTemplates/DatasmithSceneComponentTemplate.h
new file mode 100644
--- /dev/null
+++ b/Engine/Plugins/Enterprise/DatasmithContent/Source/DatasmithContent/Public/ObjectTemplates/DatasmithSceneComponentTemplate.h
@@ -0,0 +1,28 @@
+// Copyright 1998-2018 Epic Games, Inc. All Rights Reserved.
+
+#pragma once
+
+#include "Engine/EngineTypes.h"
+#include "ObjectTemplates/DatasmithObjectTemplate.h"
+
+#include "DatasmithSceneComponentTemplate.generated.h"
+
+UCLASS()
+class DATASMITHCONTENT_API UDatasmithSceneComponentTemplate : public UDatasmithObjectTemplate
+{
+	GENERATED_BODY()
+
+public:
+
+	UPROPERTY()
+	FTransform RelativeTransform;
+
+	UPROPERTY()
+	TEnumAsByte< EComponentMobility::Type > Mobility;
+
+	UPROPERTY()
+	class USceneComponent* AttachParent;
+
+	virtual void Apply( UObject* Destination, bool bForce = false ) override;
+	virtual void Load( const UObject* Source ) override;
+};
\ No newline at end of file
diff --git a/Engine/Plugins/Enterprise/DatasmithContent/Source/DatasmithContent/Public/ObjectTemplates/DatasmithSkyLightComponentTemplate.h b/Engine/Plugins/Enterprise/DatasmithContent/Source/DatasmithContent/Public/ObjectTemplates/DatasmithSkyLightComponentTemplate.h
new file mode 100644
--- /dev/null
+++ b/Engine/Plugins/Enterprise/DatasmithContent/Source/DatasmithContent/Public/ObjectTemplates/DatasmithSkyLightComponentTemplate.h
@@ -0,0 +1,30 @@
+// Copyright 1998-2018 Epic Games, Inc. All Rights Reserved.
+
+#pragma once
+
+#include "ObjectTemplates/DatasmithObjectTemplate.h"
+
+#include "Components/SkyLightComponent.h"
+
+#include "DatasmithSkyLightComponentTemplate.generated.h"
+
+UCLASS()
+class DATASMITHCONTENT_API UDatasmithSkyLightComponentTemplate : public UDatasmithObjectTemplate
+{
+	GENERATED_BODY()
+
+public:
+	UDatasmithSkyLightComponentTemplate();
+
+	UPROPERTY()
+	TEnumAsByte< ESkyLightSourceType > SourceType;
+
+	UPROPERTY()
+	int32 CubemapResolution;
+
+	UPROPERTY()
+	class UTextureCube* Cubemap;
+
+	virtual void Apply( UObject* Destination, bool bForce = false ) override;
+	virtual void Load( const UObject* Source ) override;
+};
\ No newline at end of file
diff --git a/Engine/Plugins/Enterprise/DatasmithContent/Source/DatasmithContent/Public/ObjectTemplates/DatasmithStaticMeshComponentTemplate.h b/Engine/Plugins/Enterprise/DatasmithContent/Source/DatasmithContent/Public/ObjectTemplates/DatasmithStaticMeshComponentTemplate.h
new file mode 100644
--- /dev/null
+++ b/Engine/Plugins/Enterprise/DatasmithContent/Source/DatasmithContent/Public/ObjectTemplates/DatasmithStaticMeshComponentTemplate.h
@@ -0,0 +1,24 @@
+// Copyright 1998-2018 Epic Games, Inc. All Rights Reserved.
+
+#pragma once
+
+#include "ObjectTemplates/DatasmithSceneComponentTemplate.h"
+
+#include "DatasmithStaticMeshComponentTemplate.generated.h"
+
+UCLASS()
+class DATASMITHCONTENT_API UDatasmithStaticMeshComponentTemplate : public UDatasmithSceneComponentTemplate
+{
+	GENERATED_BODY()
+
+public:
+
+	UPROPERTY()
+	class UStaticMesh* StaticMesh;
+
+	UPROPERTY()
+	TArray< class UMaterialInterface* > OverrideMaterials;
+
+	virtual void Apply( UObject* Destination, bool bForce = false ) override;
+	virtual void Load( const UObject* Source ) override;
+};
diff --git a/Engine/Plugins/Enterprise/DatasmithContent/Source/DatasmithContent/Public/ObjectTemplates/DatasmithStaticMeshTemplate.h b/Engine/Plugins/Enterprise/DatasmithContent/Source/DatasmithContent/Public/ObjectTemplates/DatasmithStaticMeshTemplate.h
new file mode 100644
--- /dev/null
+++ b/Engine/Plugins/Enterprise/DatasmithContent/Source/DatasmithContent/Public/ObjectTemplates/DatasmithStaticMeshTemplate.h
@@ -0,0 +1,128 @@
+// Copyright 1998-2018 Epic Games, Inc. All Rights Reserved.
+
+#pragma once
+
+#include "Engine/EngineTypes.h"
+#include "DatasmithObjectTemplate.h"
+
+#include "DatasmithStaticMeshTemplate.generated.h"
+
+struct FMeshSectionInfo;
+struct FMeshSectionInfoMap;
+struct FStaticMaterial;
+
+USTRUCT()
+struct DATASMITHCONTENT_API FDatasmithMeshBuildSettingsTemplate
+{
+	GENERATED_BODY()
+
+public:
+	UPROPERTY()
+	uint8 bUseMikkTSpace:1;
+
+	UPROPERTY()
+	uint8 bRecomputeNormals:1;
+
+	UPROPERTY()
+	uint8 bRecomputeTangents:1;
+
+	UPROPERTY()
+	uint8 bRemoveDegenerates:1;
+
+	UPROPERTY()
+	uint8 bBuildAdjacencyBuffer:1;
+
+	UPROPERTY()
+	uint8 bUseHighPrecisionTangentBasis:1;
+
+	UPROPERTY()
+	uint8 bUseFullPrecisionUVs:1;
+
+	UPROPERTY()
+	uint8 bGenerateLightmapUVs:1;
+
+	UPROPERTY()
+	int32 MinLightmapResolution;
+
+	UPROPERTY()
+	int32 SrcLightmapIndex;
+
+	UPROPERTY()
+	int32 DstLightmapIndex;
+
+public:
+	FDatasmithMeshBuildSettingsTemplate();
+
+	void Apply( FMeshBuildSettings* Destination, FDatasmithMeshBuildSettingsTemplate* PreviousTemplate );
+	void Load( const FMeshBuildSettings& Source );
+};
+
+USTRUCT()
+struct DATASMITHCONTENT_API FDatasmithStaticMaterialTemplate
+{
+	GENERATED_BODY()
+
+public:
+	FDatasmithStaticMaterialTemplate();
+
+	UPROPERTY()
+	FName MaterialSlotName;
+
+	UPROPERTY()
+	class UMaterialInterface* MaterialInterface;
+
+	void Apply( FStaticMaterial* Destination, FDatasmithStaticMaterialTemplate* PreviousTemplate );
+	void Load( const FStaticMaterial& Source );
+};
+
+USTRUCT()
+struct DATASMITHCONTENT_API FDatasmithMeshSectionInfoTemplate
+{
+	GENERATED_BODY()
+
+public:
+	FDatasmithMeshSectionInfoTemplate();
+
+	UPROPERTY()
+	int32 MaterialIndex;
+
+	void Apply( FMeshSectionInfo* Destination, FDatasmithMeshSectionInfoTemplate* PreviousTemplate );
+	void Load( const FMeshSectionInfo& Source );
+};
+
+USTRUCT()
+struct DATASMITHCONTENT_API FDatasmithMeshSectionInfoMapTemplate
+{
+	GENERATED_BODY()
+
+	UPROPERTY()
+	TMap< uint32, FDatasmithMeshSectionInfoTemplate > Map;
+
+	void Apply( FMeshSectionInfoMap* Destination, FDatasmithMeshSectionInfoMapTemplate* PreviousTemplate );
+	void Load( const FMeshSectionInfoMap& Source );
+};
+
+UCLASS()
+class DATASMITHCONTENT_API UDatasmithStaticMeshTemplate : public UDatasmithObjectTemplate
+{
+	GENERATED_BODY()
+
+public:
+	virtual void Apply( UObject* Destination, bool bForce = false ) override;
+	virtual void Load( const UObject* Source ) override;
+	
+	UPROPERTY( VisibleAnywhere, Category = StaticMesh )
+	FDatasmithMeshSectionInfoMapTemplate SectionInfoMap;
+
+	UPROPERTY( VisibleAnywhere, Category = StaticMesh )
+	int32 LightMapCoordinateIndex;
+
+	UPROPERTY( VisibleAnywhere, Category = StaticMesh )
+	int32 LightMapResolution;
+
+	UPROPERTY( VisibleAnywhere, Category = StaticMesh )
+	TArray< FDatasmithMeshBuildSettingsTemplate > BuildSettings;
+
+	UPROPERTY( VisibleAnywhere, Category = StaticMesh )
+	TArray< FDatasmithStaticMaterialTemplate > StaticMaterials;
+};
diff --git a/Engine/Plugins/Experimental/PythonScriptPlugin/Source/PythonScriptPlugin/Private/IncludePython.h b/Engine/Plugins/Experimental/PythonScriptPlugin/Source/PythonScriptPlugin/Private/IncludePython.h
--- a/Engine/Plugins/Experimental/PythonScriptPlugin/Source/PythonScriptPlugin/Private/IncludePython.h
+++ b/Engine/Plugins/Experimental/PythonScriptPlugin/Source/PythonScriptPlugin/Private/IncludePython.h
@@ -3,6 +3,7 @@
 #pragma once
 
 #include "CoreTypes.h"
+#include "Stats/Stats.h"
 
 #if WITH_PYTHON
 
@@ -11,4 +12,6 @@ THIRD_PARTY_INCLUDES_START
 #include "structmember.h"
 THIRD_PARTY_INCLUDES_END
 
+DECLARE_STATS_GROUP(TEXT("Python"), STATGROUP_Python, STATCAT_Advanced);
+
 #endif	// WITH_PYTHON
