Set the program name for UE4-initiated Perforce connections.

diff --git a/Engine/Plugins/Developer/PerforceSourceControl/Source/PerforceSourceControl/Private/PerforceConnection.cpp b/Engine/Plugins/Developer/PerforceSourceControl/Source/PerforceSourceControl/Private/PerforceConnection.cpp
--- a/Engine/Plugins/Developer/PerforceSourceControl/Source/PerforceSourceControl/Private/PerforceConnection.cpp
+++ b/Engine/Plugins/Developer/PerforceSourceControl/Source/PerforceSourceControl/Private/PerforceConnection.cpp
@@ -318,6 +318,7 @@ bool FPerforceConnection::EnsureValidConnection(FString& InOutServerName, FStrin
 	FString NewClientSpecName = InOutWorkspaceName;
 
 	ClientApi TestP4;
+	TestP4.SetProg("UE4");
 	TestP4.SetProtocol("tag", "");
 	TestP4.SetProtocol("enableStreams", "");
 
@@ -684,6 +685,7 @@ void FPerforceConnection::EstablishConnection(const FPerforceConnectionInfo& InC
 
 	UE_LOG(LogSourceControl, Verbose, TEXT("Attempting P4 connection: %s/%s"), *InConnectionInfo.Port, *InConnectionInfo.UserName);
 
+	P4Client.SetProg("UE4");
 	P4Client.SetProtocol("tag", "");
 	P4Client.SetProtocol("enableStreams", "");
 
diff --git a/Engine/Plugins/Developer/PerforceSourceControl/Source/PerforceSourceControl/Private/SPerforceSourceControlSettings.cpp b/Engine/Plugins/Developer/PerforceSourceControl/Source/PerforceSourceControl/Private/SPerforceSourceControlSettings.cpp
--- a/Engine/Plugins/Developer/PerforceSourceControl/Source/PerforceSourceControl/Private/SPerforceSourceControlSettings.cpp
+++ b/Engine/Plugins/Developer/PerforceSourceControl/Source/PerforceSourceControl/Private/SPerforceSourceControlSettings.cpp
@@ -31,6 +31,7 @@ void SPerforceSourceControlSettings::Construct(const FArguments& InArgs)
 	{
 #if USE_P4_API
 		ClientApi TestP4;
+		TestP4.SetProg("UE4");
 		Error P4Error;
 		TestP4.Init(&P4Error);
 		PortName = ANSI_TO_TCHAR(TestP4.GetPort().Text());
