Fixing display issue with branch in content browser

diff --git a/Engine/Plugins/Developer/PerforceSourceControl/Source/PerforceSourceControl/Private/PerforceSourceControlOperations.cpp b/Engine/Plugins/Developer/PerforceSourceControl/Source/PerforceSourceControl/Private/PerforceSourceControlOperations.cpp
--- a/Engine/Plugins/Developer/PerforceSourceControl/Source/PerforceSourceControl/Private/PerforceSourceControlOperations.cpp
+++ b/Engine/Plugins/Developer/PerforceSourceControl/Source/PerforceSourceControl/Private/PerforceSourceControlOperations.cpp
@@ -756,7 +756,7 @@ static void ParseUpdateStatusResults(const FP4RecordSet& InRecords, const TArray
 
 				int32 AtIndex = OtherOpenRecordValue.Find(TEXT("@"));
 				FString OtherOpenUser = AtIndex == INDEX_NONE ? FString(TEXT("")) : OtherOpenRecordValue.Left(AtIndex);
-				BranchModification.OtherUserCheckedOut += OtherOpenUser + TEXT(" @ ") + BranchModification.BranchName;
+				BranchModification.OtherUserCheckedOut += OtherOpenUser + TEXT(" @ ") + Branch;
 
 				if (OpenIdx < OtherOpenNum - 1)
 				{
