Never overwrite edit with an integrate for same mod time

diff --git a/Engine/Plugins/Developer/PerforceSourceControl/Source/PerforceSourceControl/Private/PerforceSourceControlOperations.cpp b/Engine/Plugins/Developer/PerforceSourceControl/Source/PerforceSourceControl/Private/PerforceSourceControlOperations.cpp
--- a/Engine/Plugins/Developer/PerforceSourceControl/Source/PerforceSourceControl/Private/PerforceSourceControlOperations.cpp
+++ b/Engine/Plugins/Developer/PerforceSourceControl/Source/PerforceSourceControl/Private/PerforceSourceControlOperations.cpp
@@ -675,17 +675,22 @@ static void ParseBranchModificationResults(const FP4RecordSet& InRecords, const
 		if (BranchModifications.Contains(BranchFile))
 		{
 			FBranchModification& BranchModification = BranchModifications[BranchFile];
-
-			// Never overwrite a current branch modification with the same from a different branch
+			
 			if (BranchModification.ModTime == HeadModTime)
 			{
+				// Never overwrite a current branch modification with the same from a different branch
 				if (BranchModification.BranchName == CurrentBranch && Branch != CurrentBranch)
 				{
 					continue;
 				}
+
+				// Never overwrite edit with an integrate for same mod time
+				if (BranchModification.Action == TEXT("edit"))
+				{
+					continue;
+				}
 			}
 
-			//  We want latest modification, <= so we catch the actual edit CL in the revision list
 			if (BranchModification.ModTime <= HeadModTime)
 			{
 				BranchModification.ModTime = HeadModTime;
