Content browser status, filter delete operations (mod time 0) when content has been readded via branch/etc
#rnx

diff --git a/Engine/Plugins/Developer/PerforceSourceControl/Source/PerforceSourceControl/Private/PerforceSourceControlOperations.cpp b/Engine/Plugins/Developer/PerforceSourceControl/Source/PerforceSourceControl/Private/PerforceSourceControlOperations.cpp
--- a/Engine/Plugins/Developer/PerforceSourceControl/Source/PerforceSourceControl/Private/PerforceSourceControlOperations.cpp
+++ b/Engine/Plugins/Developer/PerforceSourceControl/Source/PerforceSourceControl/Private/PerforceSourceControlOperations.cpp
@@ -691,6 +691,12 @@ static void ParseBranchModificationResults(const FP4RecordSet& InRecords, const
 				}
 			}
 
+			// filter deletes if file re-added
+			if (HeadAction == TEXT("delete") && BranchModification.ChangeList > HeadChange)
+			{
+				continue;
+			}
+
 			if (BranchModification.ModTime <= HeadModTime)
 			{
 				BranchModification.ModTime = HeadModTime;
@@ -870,8 +876,21 @@ static void ParseUpdateStatusResults(const FP4RecordSet& InRecords, const TArray
 
 			if (BranchModification.BranchName.Len())
 			{
-				// If the branch modification change is more recent record it
-				if (BranchModification.ModTime > State.HeadModTime)
+				bool Skip = false;
+
+				// don't record if we deleted on a status branch, though have since re-added
+				if (BranchModification.Action == TEXT("delete") && BranchModification.ChangeList < State.HeadChangeList)
+				{
+					Skip = true;
+				}
+
+				// If the branch modification change is less recent skip it
+				if (BranchModification.ModTime < State.HeadModTime)
+				{
+					Skip = true;
+				}
+				
+				if (!Skip)
 				{
 					State.HeadBranch = BranchModification.BranchName;
 					State.HeadAction = BranchModification.Action;
