Content browser, owner record branch change if mod time is newer
#rnx

diff --git a/Engine/Plugins/Developer/PerforceSourceControl/Source/PerforceSourceControl/Private/PerforceSourceControlOperations.cpp b/Engine/Plugins/Developer/PerforceSourceControl/Source/PerforceSourceControl/Private/PerforceSourceControlOperations.cpp
--- a/Engine/Plugins/Developer/PerforceSourceControl/Source/PerforceSourceControl/Private/PerforceSourceControlOperations.cpp
+++ b/Engine/Plugins/Developer/PerforceSourceControl/Source/PerforceSourceControl/Private/PerforceSourceControlOperations.cpp
@@ -885,7 +885,7 @@ static void ParseUpdateStatusResults(const FP4RecordSet& InRecords, const TArray
 				}
 
 				// If the branch modification change is less recent skip it
-				if (BranchModification.ModTime < State.HeadModTime)
+				if (BranchModification.ModTime <= State.HeadModTime)
 				{
 					Skip = true;
 				}
