Perforce source control provider now detects login expiration and displays error status in level editor toolbar

diff --git a/Engine/Plugins/Developer/PerforceSourceControl/Source/PerforceSourceControl/Private/PerforceConnection.cpp b/Engine/Plugins/Developer/PerforceSourceControl/Source/PerforceSourceControl/Private/PerforceConnection.cpp
--- a/Engine/Plugins/Developer/PerforceSourceControl/Source/PerforceSourceControl/Private/PerforceConnection.cpp
+++ b/Engine/Plugins/Developer/PerforceSourceControl/Source/PerforceSourceControl/Private/PerforceConnection.cpp
@@ -199,6 +199,8 @@ static bool TestConnection(ClientApi& P4Client, const FString& ClientSpecName, b
 	// clean up args
 	delete [] ClientSpecUTF8Name;
 
+	FPerforceSourceControlModule::SetLastErrors(OutErrorMessages);
+
 	// If there are error messages, user name is most likely invalid. Otherwise, make sure workspace actually
 	// exists on server by checking if we have it's update date.
 	bool bConnectionOK = OutErrorMessages.Num() == 0 && Records.Num() > 0 && Records[0].Contains(TEXT("Update"));
@@ -639,6 +641,18 @@ bool FPerforceConnection::RunCommand(const FString& InCommand, const TArray<FStr
 	}
 	delete [] ArgV;
 
+	// Only report connection related errors to avoid clearing of connection related error messages
+	static TSet<FString> CommandsWithoutLoginErrors;
+	if (CommandsWithoutLoginErrors.Num() == 0)
+	{
+		CommandsWithoutLoginErrors.Add(TEXT("info"));
+	}
+
+	if (!CommandsWithoutLoginErrors.Contains(InCommand))
+	{
+		FPerforceSourceControlModule::SetLastErrors(OutErrorMessage);
+	}
+
 	if (bInStandardDebugOutput)
 	{
 		UE_LOG( LogSourceControl, VeryVerbose, TEXT("P4 execution time: %0.4f seconds. Command: %s"), FPlatformTime::Seconds() - SCCStartTime, *FullCommand );
diff --git a/Engine/Plugins/Developer/PerforceSourceControl/Source/PerforceSourceControl/Private/PerforceSourceControlModule.cpp b/Engine/Plugins/Developer/PerforceSourceControl/Source/PerforceSourceControl/Private/PerforceSourceControlModule.cpp
--- a/Engine/Plugins/Developer/PerforceSourceControl/Source/PerforceSourceControl/Private/PerforceSourceControlModule.cpp
+++ b/Engine/Plugins/Developer/PerforceSourceControl/Source/PerforceSourceControl/Private/PerforceSourceControlModule.cpp
@@ -59,6 +59,15 @@ void FPerforceSourceControlModule::SaveSettings()
 	PerforceSourceControlSettings.SaveSettings();
 }
 
+void FPerforceSourceControlModule::SetLastErrors(const TArray<FText>& InErrors)
+{
+	FPerforceSourceControlModule* Module = FModuleManager::GetModulePtr<FPerforceSourceControlModule>("PerforceSourceControl");
+	if (Module)
+	{
+		Module->GetProvider().SetLastErrors(InErrors);
+	}
+}
+
 IMPLEMENT_MODULE(FPerforceSourceControlModule, PerforceSourceControl);
 
 #undef LOCTEXT_NAMESPACE
diff --git a/Engine/Plugins/Developer/PerforceSourceControl/Source/PerforceSourceControl/Private/PerforceSourceControlModule.h b/Engine/Plugins/Developer/PerforceSourceControl/Source/PerforceSourceControl/Private/PerforceSourceControlModule.h
--- a/Engine/Plugins/Developer/PerforceSourceControl/Source/PerforceSourceControl/Private/PerforceSourceControlModule.h
+++ b/Engine/Plugins/Developer/PerforceSourceControl/Source/PerforceSourceControl/Private/PerforceSourceControlModule.h
@@ -38,6 +38,9 @@ public:
 		return FModuleManager::LoadModuleChecked< FPerforceSourceControlModule >("PerforceSourceControl");
 	}
 
+	/** Set list of error messages that occurred after last perforce command */
+	static void SetLastErrors(const TArray<FText>& InErrors);
+
 private:
 	/** The one and only Perforce source control provider */
 	FPerforceSourceControlProvider PerforceSourceControlProvider;
diff --git a/Engine/Plugins/Developer/PerforceSourceControl/Source/PerforceSourceControl/Private/PerforceSourceControlProvider.cpp b/Engine/Plugins/Developer/PerforceSourceControl/Source/PerforceSourceControl/Private/PerforceSourceControlProvider.cpp
--- a/Engine/Plugins/Developer/PerforceSourceControl/Source/PerforceSourceControl/Private/PerforceSourceControlProvider.cpp
+++ b/Engine/Plugins/Developer/PerforceSourceControl/Source/PerforceSourceControl/Private/PerforceSourceControlProvider.cpp
@@ -76,7 +76,19 @@ FText FPerforceSourceControlProvider::GetStatusText() const
 	Args.Add( TEXT("UserName"), FText::FromString( Settings.GetUserName() ) );
 	Args.Add( TEXT("ClientSpecName"), FText::FromString( Settings.GetWorkspace() ) );
 
-	return FText::Format( LOCTEXT("PerforceStatusText", "Enabled: {IsEnabled}\nConnected: {IsConnected}\n\nPort: {PortNumber}\nUser name: {UserName}\nClient name: {ClientSpecName}"), Args );
+	FText FormattedError;
+	TArray<FText> RecentErrors = GetLastErrors();
+	if (RecentErrors.Num() > 0)
+	{
+		FFormatNamedArguments ErrorArgs;
+		ErrorArgs.Add( TEXT("ErrorText"), RecentErrors[0] );
+
+		FormattedError = FText::Format( LOCTEXT("PerforceErrorStatusText", "Error: {ErrorText}\n\n"), ErrorArgs );
+	}
+
+	Args.Add( TEXT("ErrorText"), FormattedError);
+
+	return FText::Format( LOCTEXT("PerforceStatusText", "{ErrorText}Enabled: {IsEnabled}\nConnected: {IsConnected}\n\nPort: {PortNumber}\nUser name: {UserName}\nClient name: {ClientSpecName}"), Args );
 }
 
 bool FPerforceSourceControlProvider::IsEnabled() const
@@ -86,7 +98,7 @@ bool FPerforceSourceControlProvider::IsEnabled() const
 
 bool FPerforceSourceControlProvider::IsAvailable() const
 {
-	return bServerAvailable;
+	return bServerAvailable && !IsLoginError();
 }
 
 bool FPerforceSourceControlProvider::EstablishPersistentConnection()
@@ -147,6 +159,7 @@ void FPerforceSourceControlProvider::ParseCommandLineSettings(bool bForceConnect
 
 	if (bForceConnection)
 	{
+		bLoginError = false;
 		FPerforceConnectionInfo ConnectionInfo = P4Settings.GetConnectionInfo();
 		if(FPerforceConnection::EnsureValidConnection(PortName, UserName, ClientSpecName, ConnectionInfo))
 		{
@@ -180,6 +193,44 @@ const FName& FPerforceSourceControlProvider::GetName() const
 	return ProviderName;
 }
 
+void FPerforceSourceControlProvider::SetLastErrors(const TArray<FText>& InErrors)
+{
+	static FString SessionExpiredMessage(TEXT("Your session has expired, please login again.\n"));
+
+	bool bContainsLoginError = false;
+	for (const FText& It : InErrors)
+	{
+		if (It.ToString() == SessionExpiredMessage)
+		{
+			bContainsLoginError = true;
+			break;
+		}
+	}
+
+	bLoginError = bContainsLoginError;
+
+	FScopeLock Lock(&LastErrorsCriticalSection);
+	LastErrors = InErrors;
+}
+
+bool FPerforceSourceControlProvider::IsLoginError() const
+{
+	return bLoginError;
+}
+
+TArray<FText> FPerforceSourceControlProvider::GetLastErrors() const
+{
+	FScopeLock Lock(&LastErrorsCriticalSection);
+	TArray<FText> Result = LastErrors;
+	return Result;
+}
+
+int32 FPerforceSourceControlProvider::GetNumLastErrors() const
+{
+	FScopeLock Lock(&LastErrorsCriticalSection);
+	return LastErrors.Num();
+}
+
 ECommandResult::Type FPerforceSourceControlProvider::GetState( const TArray<FString>& InFiles, TArray< TSharedRef<ISourceControlState, ESPMode::ThreadSafe> >& OutState, EStateCacheUsage::Type InStateCacheUsage )
 {
 	if(!IsEnabled())
diff --git a/Engine/Plugins/Developer/PerforceSourceControl/Source/PerforceSourceControl/Private/PerforceSourceControlProvider.h b/Engine/Plugins/Developer/PerforceSourceControl/Source/PerforceSourceControl/Private/PerforceSourceControlProvider.h
--- a/Engine/Plugins/Developer/PerforceSourceControl/Source/PerforceSourceControl/Private/PerforceSourceControlProvider.h
+++ b/Engine/Plugins/Developer/PerforceSourceControl/Source/PerforceSourceControl/Private/PerforceSourceControlProvider.h
@@ -19,6 +19,7 @@ public:
 	/** Constructor */
 	FPerforceSourceControlProvider()
 		: bServerAvailable(false)
+		, bLoginError(false)
 		, PersistentConnection(NULL)
 	{
 	}
@@ -67,6 +68,18 @@ public:
 	/** Get the P4 ticket we will use for connections */
 	const FString& GetTicket() const;
 
+	/** Set list of error messages that occurred after last perforce command */
+	void SetLastErrors(const TArray<FText>& InErrors);
+
+	/** Did most recent command generate a login error */
+	bool IsLoginError() const;
+
+	/** Get list of error messages that occurred after last perforce command */
+	TArray<FText> GetLastErrors() const;
+
+	/** Get number of error messages seen after running last perforce command */
+	int32 GetNumLastErrors() const;
+
 	/** Helper function used to update state cache */
 	TSharedRef<FPerforceSourceControlState, ESPMode::ThreadSafe> GetStateInternal(const FString& InFilename);
 
@@ -124,6 +137,15 @@ private:
 	/** Indicates if source control integration is available or not. */
 	bool bServerAvailable;
 
+	/** Saw login error when running last command. */
+	TAtomic<bool> bLoginError;
+
+	/** Critical section for thread safety of error messages that occurred after last perforce command */
+	mutable FCriticalSection LastErrorsCriticalSection;
+
+	/** List of error messages that occurred after last perforce command */
+	TArray<FText> LastErrors;
+
 	/** A pointer to the persistent P4 connection for synchronous operations */
 	class FPerforceConnection* PersistentConnection;
 
