Copying from Dev-Editor
PR #6608: Perforce Plugin: Use -P4Passwd cmd argument value when connecting to Perforce. (Contributed by [REDACTED_PERSON])

diff --git a/Engine/Plugins/Developer/PerforceSourceControl/Source/PerforceSourceControl/Private/PerforceConnection.cpp b/Engine/Plugins/Developer/PerforceSourceControl/Source/PerforceSourceControl/Private/PerforceConnection.cpp
--- a/Engine/Plugins/Developer/PerforceSourceControl/Source/PerforceSourceControl/Private/PerforceConnection.cpp
+++ b/Engine/Plugins/Developer/PerforceSourceControl/Source/PerforceSourceControl/Private/PerforceConnection.cpp
@@ -383,6 +383,7 @@ bool FPerforceConnection::EnsureValidConnection(FString& InOutServerName, FStrin
 			TestP4.SetCwd(FROM_TCHAR(*FPaths::RootDir(), bIsUnicodeServer));
 			TestP4.SetUser(FROM_TCHAR(*NewUserName, bIsUnicodeServer));
 			TestP4.SetClient(FROM_TCHAR(*NewClientSpecName, bIsUnicodeServer));
+			TestP4.SetPassword(FROM_TCHAR(*InConnectionInfo.Ticket, bIsUnicodeServer));
 
 		}
 	}
@@ -753,7 +754,11 @@ void FPerforceConnection::EstablishConnection(const FPerforceConnectionInfo& InC
 			// Now we know our unicode status we can gather the client root
 			P4Client.SetUser(FROM_TCHAR(*InConnectionInfo.UserName, bIsUnicode));
 
-			if(InConnectionInfo.Password.Len() > 0)
+			if (InConnectionInfo.Ticket.Len() > 0)
+			{
+				P4Client.SetPassword(FROM_TCHAR(*InConnectionInfo.Ticket, bIsUnicode));
+			}
+			else if (InConnectionInfo.Password.Len() > 0)
 			{
 				Login(InConnectionInfo);
 			}
