Fix perforce connection icon not appering correctly if a server and credentials are specified but the server cannot be contacted.

diff --git a/Engine/Plugins/Developer/PerforceSourceControl/Source/PerforceSourceControl/Private/PerforceConnection.cpp b/Engine/Plugins/Developer/PerforceSourceControl/Source/PerforceSourceControl/Private/PerforceConnection.cpp
--- a/Engine/Plugins/Developer/PerforceSourceControl/Source/PerforceSourceControl/Private/PerforceConnection.cpp
+++ b/Engine/Plugins/Developer/PerforceSourceControl/Source/PerforceSourceControl/Private/PerforceConnection.cpp
@@ -792,7 +792,10 @@ FScopedPerforceConnection::FScopedPerforceConnection( FPerforceSourceControlComm
 	, Concurrency(InCommand.Concurrency)
 {
 	Initialize(InCommand.ConnectionInfo);
-	InCommand.MarkConnectionAsSuccessful();
+	if (IsValid())
+	{
+		InCommand.MarkConnectionAsSuccessful();
+	}
 }
 
 FScopedPerforceConnection::FScopedPerforceConnection( EConcurrency::Type InConcurrency, const FPerforceConnectionInfo& InConnectionInfo )
diff --git a/Engine/Plugins/Developer/PerforceSourceControl/Source/PerforceSourceControl/Private/PerforceSourceControlProvider.cpp b/Engine/Plugins/Developer/PerforceSourceControl/Source/PerforceSourceControl/Private/PerforceSourceControlProvider.cpp
--- a/Engine/Plugins/Developer/PerforceSourceControl/Source/PerforceSourceControl/Private/PerforceSourceControlProvider.cpp
+++ b/Engine/Plugins/Developer/PerforceSourceControl/Source/PerforceSourceControl/Private/PerforceSourceControlProvider.cpp
@@ -417,7 +417,7 @@ void FPerforceSourceControlProvider::Tick()
 			CommandQueue.RemoveAt(CommandIndex);
 
 			// update connection state
-			bServerAvailable = !Command.bConnectionDropped || Command.bCancelled;
+			bServerAvailable = Command.bConnectionWasSuccessful && (!Command.bConnectionDropped || Command.bCancelled);
 
 			// let command update the states of any files
 			bStatesUpdated |= Command.Worker->UpdateStates();
