Fix a thread safety issue in perforce integration

diff --git a/Engine/Plugins/Developer/PerforceSourceControl/Source/PerforceSourceControl/Private/PerforceConnection.cpp b/Engine/Plugins/Developer/PerforceSourceControl/Source/PerforceSourceControl/Private/PerforceConnection.cpp
--- a/Engine/Plugins/Developer/PerforceSourceControl/Source/PerforceSourceControl/Private/PerforceConnection.cpp
+++ b/Engine/Plugins/Developer/PerforceSourceControl/Source/PerforceSourceControl/Private/PerforceConnection.cpp
@@ -643,13 +643,7 @@ bool FPerforceConnection::RunCommand(const FString& InCommand, const TArray<FStr
 	delete [] ArgV;
 
 	// Only report connection related errors to avoid clearing of connection related error messages
-	static TSet<FString> CommandsWithoutLoginErrors;
-	if (CommandsWithoutLoginErrors.Num() == 0)
-	{
-		CommandsWithoutLoginErrors.Add(TEXT("info"));
-	}
-
-	if (!CommandsWithoutLoginErrors.Contains(InCommand))
+	if (InCommand != TEXT("info"))
 	{
 		FPerforceSourceControlModule::SetLastErrors(OutErrorMessage);
 	}
