Changelist window - Drag & drop support

diff --git a/Engine/Plugins/Developer/PerforceSourceControl/Source/PerforceSourceControl/Private/PerforceSourceControlChangelist.h b/Engine/Plugins/Developer/PerforceSourceControl/Source/PerforceSourceControl/Private/PerforceSourceControlChangelist.h
--- a/Engine/Plugins/Developer/PerforceSourceControl/Source/PerforceSourceControl/Private/PerforceSourceControlChangelist.h
+++ b/Engine/Plugins/Developer/PerforceSourceControl/Source/PerforceSourceControl/Private/PerforceSourceControlChangelist.h
@@ -10,7 +10,7 @@ class FPerforceSourceControlChangelist : public ISourceControlChangelist
 public:
 	FPerforceSourceControlChangelist();
 
-	FPerforceSourceControlChangelist(int32 InChangelistNumber) 
+	explicit FPerforceSourceControlChangelist(int32 InChangelistNumber)
 		: ChangelistNumber(InChangelistNumber)
 	{
 	}
@@ -25,6 +25,16 @@ public:
 		return ChangelistNumber == InOther.ChangelistNumber;
 	}
 
+	bool operator!=(const FPerforceSourceControlChangelist& InOther) const
+	{
+		return ChangelistNumber != InOther.ChangelistNumber;
+	}
+
+	bool IsDefault() const
+	{
+		return *this == DefaultChangelist;
+	}
+
 	friend FORCEINLINE uint32 GetTypeHash(const FPerforceSourceControlChangelist& PerforceChangelist)
 	{
 		return GetTypeHash(PerforceChangelist.ChangelistNumber);
diff --git a/Engine/Plugins/Developer/PerforceSourceControl/Source/PerforceSourceControl/Private/PerforceSourceControlCommand.cpp b/Engine/Plugins/Developer/PerforceSourceControl/Source/PerforceSourceControl/Private/PerforceSourceControlCommand.cpp
--- a/Engine/Plugins/Developer/PerforceSourceControl/Source/PerforceSourceControl/Private/PerforceSourceControlCommand.cpp
+++ b/Engine/Plugins/Developer/PerforceSourceControl/Source/PerforceSourceControl/Private/PerforceSourceControlCommand.cpp
@@ -17,7 +17,7 @@ FPerforceSourceControlCommand::FPerforceSourceControlCommand(const TSharedRef<cl
 	, bConnectionDropped(false)
 	, bAutoDelete(true)
 	, Concurrency(EConcurrency::Synchronous)
-	, Changelist(FPerforceSourceControlChangelist::DefaultChangelist.ChangelistNumber)
+	, Changelist(FPerforceSourceControlChangelist::DefaultChangelist)
 {
 	// grab the providers settings here, so we don't access them once the worker thread is launched
 	check(IsInGameThread());
diff --git a/Engine/Plugins/Developer/PerforceSourceControl/Source/PerforceSourceControl/Private/PerforceSourceControlCommand.h b/Engine/Plugins/Developer/PerforceSourceControl/Source/PerforceSourceControl/Private/PerforceSourceControlCommand.h
--- a/Engine/Plugins/Developer/PerforceSourceControl/Source/PerforceSourceControl/Private/PerforceSourceControlCommand.h
+++ b/Engine/Plugins/Developer/PerforceSourceControl/Source/PerforceSourceControl/Private/PerforceSourceControlCommand.h
@@ -5,6 +5,7 @@
 #include "CoreMinimal.h"
 #include "ISourceControlProvider.h"
 #include "PerforceConnectionInfo.h"
+#include "PerforceSourceControlChangelist.h"
 #include "Misc/IQueuedWork.h"
 
 DECLARE_DELEGATE_RetVal(bool, FOnIsCancelled);
@@ -100,7 +101,7 @@ public:
 	TArray< FString > Files;
 
 	/** Changelist to perform this operation on */
-	int32 Changelist;
+	FPerforceSourceControlChangelist Changelist;
 
 	/** Potential error, warning and info message storage */
 	FSourceControlResultInfo ResultInfo;
diff --git a/Engine/Plugins/Developer/PerforceSourceControl/Source/PerforceSourceControl/Private/PerforceSourceControlModule.cpp b/Engine/Plugins/Developer/PerforceSourceControl/Source/PerforceSourceControl/Private/PerforceSourceControlModule.cpp
--- a/Engine/Plugins/Developer/PerforceSourceControl/Source/PerforceSourceControl/Private/PerforceSourceControlModule.cpp
+++ b/Engine/Plugins/Developer/PerforceSourceControl/Source/PerforceSourceControl/Private/PerforceSourceControlModule.cpp
@@ -33,6 +33,7 @@ void FPerforceSourceControlModule::StartupModule()
 	PerforceSourceControlProvider.RegisterWorker( "DeleteChangelist", FGetPerforceSourceControlWorker::CreateStatic( &CreateWorker<FPerforceDeleteChangelistWorker> ) );
 	PerforceSourceControlProvider.RegisterWorker( "EditChangelist", FGetPerforceSourceControlWorker::CreateStatic( &CreateWorker<FPerforceEditChangelistWorker> ) );
 	PerforceSourceControlProvider.RegisterWorker( "RevertUnchanged", FGetPerforceSourceControlWorker::CreateStatic( &CreateWorker<FPerforceRevertUnchangedWorker> ) );
+	PerforceSourceControlProvider.RegisterWorker( "MoveToChangelist", FGetPerforceSourceControlWorker::CreateStatic( &CreateWorker<FPerforceReopenWorker>) );
 
 	// load our settings
 	PerforceSourceControlSettings.LoadSettings();
diff --git a/Engine/Plugins/Developer/PerforceSourceControl/Source/PerforceSourceControl/Private/PerforceSourceControlOperations.cpp b/Engine/Plugins/Developer/PerforceSourceControl/Source/PerforceSourceControl/Private/PerforceSourceControlOperations.cpp
--- a/Engine/Plugins/Developer/PerforceSourceControl/Source/PerforceSourceControl/Private/PerforceSourceControlOperations.cpp
+++ b/Engine/Plugins/Developer/PerforceSourceControl/Source/PerforceSourceControl/Private/PerforceSourceControlOperations.cpp
@@ -338,6 +338,53 @@ static FText ParseSubmitResults(const FP4RecordSet& InRecords)
 	return LOCTEXT("SubmitMessageUnknown", "Submitted changelist");
 }
 
+static bool RunReopenCommand(FPerforceSourceControlCommand& InCommand, const TArray<FString>& InFiles, const FPerforceSourceControlChangelist& InChangelist, TArray<FString>* OutReopenedFiles = nullptr)
+{
+	bool bCommandSuccessful = true;
+
+	FScopedPerforceConnection ScopedConnection(InCommand);
+	if (!InCommand.IsCanceled() && ScopedConnection.IsValid())
+	{
+		FPerforceConnection& Connection = ScopedConnection.GetConnection();
+
+		// Batch reopen into multiple commands, to avoid command line limits
+		const int32 BatchedCount = 100;
+
+		if (OutReopenedFiles != nullptr)
+		{
+			OutReopenedFiles->Reserve(InFiles.Num());
+		}
+
+		for (int32 StartingIndex = 0; StartingIndex < InFiles.Num() && bCommandSuccessful; StartingIndex += BatchedCount)
+		{
+			FP4RecordSet Records;
+			TArray<FString> ReopenParams;
+
+			//Add changelist information to params
+			ReopenParams.Add(TEXT("-c"));
+			ReopenParams.Add(InChangelist.ToString());
+
+			int32 NextIndex = FMath::Min(StartingIndex + BatchedCount, InFiles.Num());
+
+			for (int32 FileIndex = StartingIndex; FileIndex < NextIndex; FileIndex++)
+			{
+				ReopenParams.Add(InFiles[FileIndex]);
+			}
+
+			bCommandSuccessful = Connection.RunCommand(TEXT("reopen"), ReopenParams, Records, InCommand.ResultInfo.ErrorMessages, FOnIsCancelled::CreateRaw(&InCommand, &FPerforceSourceControlCommand::IsCanceled), InCommand.bConnectionDropped);
+			if (bCommandSuccessful && OutReopenedFiles != nullptr)
+			{
+				for (int32 FileIndex = StartingIndex; FileIndex < NextIndex; FileIndex++)
+				{
+					OutReopenedFiles->Add(InFiles[FileIndex]);
+				}
+			}
+		}
+	}
+
+	return bCommandSuccessful;
+}
+
 bool FPerforceCheckInWorker::Execute(FPerforceSourceControlCommand& InCommand)
 {
 	FScopedPerforceConnection ScopedConnection(InCommand);
@@ -348,45 +395,18 @@ bool FPerforceCheckInWorker::Execute(FPerforceSourceControlCommand& InCommand)
 		check(InCommand.Operation->GetName() == GetName());
 		TSharedRef<FCheckIn, ESPMode::ThreadSafe> Operation = StaticCastSharedRef<FCheckIn>(InCommand.Operation);
 
-		const bool bCheckInFromDefaultChangelist = (InCommand.Changelist == FPerforceSourceControlChangelist::DefaultChangelist.ChangelistNumber);
-		const int32 BatchedCount = 100;
-		int32 ChangeList = InCommand.Changelist;
+		FPerforceSourceControlChangelist ChangeList(InCommand.Changelist);
 		TArray<FString> ReopenedFiles;
 
 		InCommand.bCommandSuccessful = true;
 
-		if (bCheckInFromDefaultChangelist)
+		if (InCommand.Changelist.IsDefault())
 		{
-			ChangeList = Connection.CreatePendingChangelist(Operation->GetDescription(), TArray<FString>(), FOnIsCancelled::CreateRaw(&InCommand, &FPerforceSourceControlCommand::IsCanceled), InCommand.ResultInfo.ErrorMessages);
-			if (ChangeList > 0)
+			int32 NewChangeList = Connection.CreatePendingChangelist(Operation->GetDescription(), TArray<FString>(), FOnIsCancelled::CreateRaw(&InCommand, &FPerforceSourceControlCommand::IsCanceled), InCommand.ResultInfo.ErrorMessages);
+			if (NewChangeList > 0)
 			{
-				// Batch reopen into multiple commands, to avoid command line limits
-				ReopenedFiles.Reserve(InCommand.Files.Num());
-				for (int32 StartingIndex = 0; StartingIndex < InCommand.Files.Num() && InCommand.bCommandSuccessful; StartingIndex += BatchedCount)
-				{
-					FP4RecordSet Records;
-					TArray< FString > ReopenParams;
-
-					//Add changelist information to params
-					ReopenParams.Insert(TEXT("-c"), 0);
-					ReopenParams.Insert(FString::Printf(TEXT("%d"), ChangeList), 1);
-					int32 NextIndex = FMath::Min(StartingIndex + BatchedCount, InCommand.Files.Num());
-
-					int32 FileParamsStartIdx = ReopenParams.Num();
-					for (int32 FileIndex = StartingIndex; FileIndex < NextIndex; FileIndex++)
-					{
-						ReopenParams.Add(InCommand.Files[FileIndex]);
-					}
-
-					InCommand.bCommandSuccessful = Connection.RunCommand(TEXT("reopen"), ReopenParams, Records, InCommand.ResultInfo.ErrorMessages, FOnIsCancelled::CreateRaw(&InCommand, &FPerforceSourceControlCommand::IsCanceled), InCommand.bConnectionDropped);
-					if (InCommand.bCommandSuccessful)
-					{
-						for (int32 ParamIdx = FileParamsStartIdx; ParamIdx < ReopenParams.Num(); ++ParamIdx)
-						{
-							ReopenedFiles.Add(ReopenParams[ParamIdx]);
-						}
-					}
-				}
+				ChangeList = FPerforceSourceControlChangelist(NewChangeList);
+				InCommand.bCommandSuccessful = RunReopenCommand(InCommand, InCommand.Files, ChangeList, &ReopenedFiles);
 			}
 			else
 			{
@@ -400,8 +420,8 @@ bool FPerforceCheckInWorker::Execute(FPerforceSourceControlCommand& InCommand)
 			TArray<FString> SubmitParams;
 			FP4RecordSet Records;
 
-			SubmitParams.Insert(TEXT("-c"), 0);
-			SubmitParams.Insert(FString::Printf(TEXT("%d"), ChangeList), 1);
+			SubmitParams.Add(TEXT("-c"));
+			SubmitParams.Add(ChangeList.ToString());
 
 			InCommand.bCommandSuccessful = Connection.RunCommand(TEXT("submit"), SubmitParams, Records, InCommand.ResultInfo.ErrorMessages, FOnIsCancelled::CreateRaw(&InCommand, &FPerforceSourceControlCommand::IsCanceled), InCommand.bConnectionDropped);
 
@@ -428,7 +448,7 @@ bool FPerforceCheckInWorker::Execute(FPerforceSourceControlCommand& InCommand)
 
 				StaticCastSharedRef<FCheckIn>(InCommand.Operation)->SetSuccessMessage(ParseSubmitResults(Records));
 
-				for (auto Iter(InCommand.Files.CreateIterator()); Iter; Iter++)
+				for(auto Iter(InCommand.Files.CreateIterator()); Iter; Iter++)
 				{
 					OutResults.Add(*Iter, EPerforceState::ReadOnly);
 				}
@@ -436,30 +456,12 @@ bool FPerforceCheckInWorker::Execute(FPerforceSourceControlCommand& InCommand)
 		}
 
 		// If the submit failed, clean up the changelist created above
-		if (!InCommand.bCommandSuccessful && bCheckInFromDefaultChangelist)
+		if (!InCommand.bCommandSuccessful && InCommand.Changelist.IsDefault())
 		{
 			// Reopen the assets to the default changelist to remove them from the changelist we created above
 			if (ReopenedFiles.Num() > 0)
 			{
-				bool bReopenSuccessful = true;
-				for (int32 StartingIndex = 0; StartingIndex < ReopenedFiles.Num() && bReopenSuccessful; StartingIndex += BatchedCount)
-				{
-					FP4RecordSet Records;
-					TArray< FString > ReopenParams;
-
-					//Add changelist information to params
-					ReopenParams.Insert(TEXT("-c"), 0);
-					ReopenParams.Insert(TEXT("default"), 1);
-					int32 NextIndex = FMath::Min(StartingIndex + BatchedCount, ReopenedFiles.Num());
-
-					int32 FileParamsStartIdx = ReopenParams.Num();
-					for (int32 FileIndex = StartingIndex; FileIndex < NextIndex; FileIndex++)
-					{
-						ReopenParams.Add(ReopenedFiles[FileIndex]);
-					}
-
-					bReopenSuccessful = Connection.RunCommand(TEXT("reopen"), ReopenParams, Records, InCommand.ResultInfo.ErrorMessages, FOnIsCancelled::CreateRaw(&InCommand, &FPerforceSourceControlCommand::IsCanceled), InCommand.bConnectionDropped);
-				}
+				InCommand.bCommandSuccessful = RunReopenCommand(InCommand, ReopenedFiles, InCommand.Changelist);
 			}
 
 			// Delete the changelist we created above
@@ -467,7 +469,7 @@ bool FPerforceCheckInWorker::Execute(FPerforceSourceControlCommand& InCommand)
 				FP4RecordSet Records;
 				TArray<FString> ChangeParams;
 				ChangeParams.Add(TEXT("-d"));
-				ChangeParams.Add(FString::Printf(TEXT("%d"), ChangeList));
+				ChangeParams.Add(ChangeList.ToString());
 				Connection.RunCommand(TEXT("change"), ChangeParams, Records, InCommand.ResultInfo.ErrorMessages, FOnIsCancelled::CreateRaw(&InCommand, &FPerforceSourceControlCommand::IsCanceled), InCommand.bConnectionDropped);
 			}
 		}
@@ -1780,7 +1782,7 @@ FName FPerforceDeleteChangelistWorker::GetName() const
 bool FPerforceDeleteChangelistWorker::Execute(class FPerforceSourceControlCommand& InCommand)
 {
 	// Can't delete the default changelist
-	if (InCommand.Changelist == FPerforceSourceControlChangelist::DefaultChangelist.ChangelistNumber)
+	if (InCommand.Changelist.IsDefault())
 	{
 		InCommand.bCommandSuccessful = false;
 		return false;
@@ -1796,7 +1798,7 @@ bool FPerforceDeleteChangelistWorker::Execute(class FPerforceSourceControlComman
 		FP4RecordSet Records;
 		TArray<FString> Params;
 		Params.Add(TEXT("-d"));
-		Params.Add(FString::Printf(TEXT("%d"), InCommand.Changelist));
+		Params.Add(InCommand.Changelist.ToString());
 		// Command will fail if changelist is not empty
 		Connection.RunCommand(TEXT("change"), Params, Records, InCommand.ResultInfo.ErrorMessages, FOnIsCancelled::CreateRaw(&InCommand, &FPerforceSourceControlCommand::IsCanceled), InCommand.bConnectionDropped);
 		// The normal parsing of the records here will show that it failed, but there's no record on a deleted changelist
@@ -1805,7 +1807,7 @@ bool FPerforceDeleteChangelistWorker::Execute(class FPerforceSourceControlComman
 		// Keep track of changelist to update the cache
 		if (InCommand.bCommandSuccessful)
 		{
-			DeletedChangelist.ChangelistNumber = InCommand.Changelist;
+			DeletedChangelist = InCommand.Changelist;
 		}
 	}
 
@@ -1815,7 +1817,7 @@ bool FPerforceDeleteChangelistWorker::Execute(class FPerforceSourceControlComman
 bool FPerforceDeleteChangelistWorker::UpdateStates() const
 {
 	FPerforceSourceControlModule& PerforceSourceControl = FPerforceSourceControlModule::Get();
-	if (DeletedChangelist.ChangelistNumber != FPerforceSourceControlChangelist::DefaultChangelist.ChangelistNumber)
+	if (!DeletedChangelist.IsDefault())
 	{
 		return PerforceSourceControl.GetProvider().RemoveChangelistFromCache(DeletedChangelist);
 	}
@@ -1840,13 +1842,13 @@ bool FPerforceEditChangelistWorker::Execute(class FPerforceSourceControlCommand&
 		check(InCommand.Operation->GetName() == GetName());
 		TSharedRef<FEditChangelist, ESPMode::ThreadSafe> Operation = StaticCastSharedRef<FEditChangelist>(InCommand.Operation);
 
-		int32 Changelist = Connection.EditPendingChangelist(Operation->GetDescription(), InCommand.Changelist, FOnIsCancelled::CreateRaw(&InCommand, &FPerforceSourceControlCommand::IsCanceled), InCommand.ResultInfo.ErrorMessages);
+		int32 ChangelistNumber = Connection.EditPendingChangelist(Operation->GetDescription(), InCommand.Changelist.ChangelistNumber, FOnIsCancelled::CreateRaw(&InCommand, &FPerforceSourceControlCommand::IsCanceled), InCommand.ResultInfo.ErrorMessages);
 
-		InCommand.bCommandSuccessful = (Changelist == InCommand.Changelist);
+		InCommand.bCommandSuccessful = (ChangelistNumber == InCommand.Changelist.ChangelistNumber);
 
 		if (InCommand.bCommandSuccessful)
 		{
-			EditedChangelist.ChangelistNumber = InCommand.Changelist;
+			EditedChangelist = InCommand.Changelist;
 			EditedDescription = Operation->GetDescription();
 		}
 	}
@@ -1879,15 +1881,7 @@ bool FPerforceRevertUnchangedWorker::Execute(class FPerforceSourceControlCommand
 
 		Parameters.Add(TEXT("-a")); // revert unchanged only
 		Parameters.Add(TEXT("-c"));
-
-		if (InCommand.Changelist != FPerforceSourceControlChangelist::DefaultChangelist.ChangelistNumber)
-		{
-			Parameters.Add(FString::Printf(TEXT("%d"), InCommand.Changelist));
-		}
-		else
-		{
-			Parameters.Add(TEXT("default"));
-		}
+		Parameters.Add(InCommand.Changelist.ToString());
 
 		if (InCommand.Files.Num() > 0)
 		{
@@ -1910,4 +1904,26 @@ bool FPerforceRevertUnchangedWorker::UpdateStates() const
 	return UpdateCachedStates(OutResults);
 }
 
+FName FPerforceReopenWorker::GetName() const
+{
+	return "Reopen";
+}
+	
+bool FPerforceReopenWorker::Execute(FPerforceSourceControlCommand& InCommand)
+{
+	FScopedPerforceConnection ScopedConnection(InCommand);
+	if (!InCommand.IsCanceled() && ScopedConnection.IsValid())
+	{
+		InCommand.bCommandSuccessful = RunReopenCommand(InCommand, InCommand.Files, InCommand.Changelist);
+	}
+
+	return InCommand.bCommandSuccessful;
+}
+
+bool FPerforceReopenWorker::UpdateStates() const
+{
+	// TODO : update changelists states (source & destination CLs)
+	return true;
+}
+
 #undef LOCTEXT_NAMESPACE
diff --git a/Engine/Plugins/Developer/PerforceSourceControl/Source/PerforceSourceControl/Private/PerforceSourceControlOperations.h b/Engine/Plugins/Developer/PerforceSourceControl/Source/PerforceSourceControl/Private/PerforceSourceControlOperations.h
--- a/Engine/Plugins/Developer/PerforceSourceControl/Source/PerforceSourceControl/Private/PerforceSourceControlOperations.h
+++ b/Engine/Plugins/Developer/PerforceSourceControl/Source/PerforceSourceControl/Private/PerforceSourceControlOperations.h
@@ -247,3 +247,14 @@ protected:
 	/** Map of filenames to perforce state */
 	TMap<FString, EPerforceState::Type> OutResults;
 };
+
+class FPerforceReopenWorker : public IPerforceSourceControlWorker
+{
+public:
+	virtual ~FPerforceReopenWorker() {}
+	// IPerforceSourceControlWorker interface
+	virtual FName GetName() const override;
+	virtual bool Execute(class FPerforceSourceControlCommand& InCommand) override;
+	virtual bool UpdateStates() const override;
+};
+
diff --git a/Engine/Plugins/Developer/PerforceSourceControl/Source/PerforceSourceControl/Private/PerforceSourceControlProvider.cpp b/Engine/Plugins/Developer/PerforceSourceControl/Source/PerforceSourceControl/Private/PerforceSourceControlProvider.cpp
--- a/Engine/Plugins/Developer/PerforceSourceControl/Source/PerforceSourceControl/Private/PerforceSourceControlProvider.cpp
+++ b/Engine/Plugins/Developer/PerforceSourceControl/Source/PerforceSourceControl/Private/PerforceSourceControlProvider.cpp
@@ -359,7 +359,7 @@ ECommandResult::Type FPerforceSourceControlProvider::Execute( const FSourceContr
 		return ECommandResult::Failed;
 	}
 
-	int32 ChangelistNumber = (InChangelist ? InChangelist->ChangelistNumber : FPerforceSourceControlChangelist::DefaultChangelist.ChangelistNumber);
+	FPerforceSourceControlChangelist Changelist = InChangelist ? InChangelist.ToSharedRef().Get() : FPerforceSourceControlChangelist::DefaultChangelist;
 
 	// fire off operation
 	if(InConcurrency == EConcurrency::Synchronous)
@@ -370,7 +370,7 @@ ECommandResult::Type FPerforceSourceControlProvider::Execute( const FSourceContr
 		Command->StatusBranchNames = StatusBranchNames;
 		Command->ContentRoot = ContentRoot;
 		Command->OperationCompleteDelegate = InOperationCompleteDelegate;
-		Command->Changelist = ChangelistNumber;
+		Command->Changelist = Changelist;
 		return ExecuteSynchronousCommand(*Command, InOperation->GetInProgressString(), true);
 	}
 	else
@@ -381,7 +381,7 @@ ECommandResult::Type FPerforceSourceControlProvider::Execute( const FSourceContr
 		Command->StatusBranchNames = StatusBranchNames;
 		Command->ContentRoot = ContentRoot;
 		Command->OperationCompleteDelegate = InOperationCompleteDelegate;
-		Command->Changelist = ChangelistNumber;
+		Command->Changelist = Changelist;
 		return IssueCommand(*Command, false);
 	}
 }
diff --git a/Engine/Source/Developer/SourceControl/Public/SourceControlOperations.h b/Engine/Source/Developer/SourceControl/Public/SourceControlOperations.h
--- a/Engine/Source/Developer/SourceControl/Public/SourceControlOperations.h
+++ b/Engine/Source/Developer/SourceControl/Public/SourceControlOperations.h
@@ -517,4 +517,23 @@ public:
 	}
 };
 
+/**
+ * Operation used to move files between changelists
+ */
+class FMoveToChangelist : public FSourceControlOperationBase
+{
+public:
+	// ISourceControlOperation interface
+	virtual FName GetName() const override
+	{
+		return "MoveToChangelist";
+	}
+
+	virtual FText GetInProgressString() const override
+	{
+		return LOCTEXT("SourceControl_MoveToChangelist", "Moving files to target changelist...");
+	}
+};
+
+
 #undef LOCTEXT_NAMESPACE
diff --git a/Engine/Source/Editor/SourceControlWindows/Private/SSourceControlChangelists.cpp b/Engine/Source/Editor/SourceControlWindows/Private/SSourceControlChangelists.cpp
--- a/Engine/Source/Editor/SourceControlWindows/Private/SSourceControlChangelists.cpp
+++ b/Engine/Source/Editor/SourceControlWindows/Private/SSourceControlChangelists.cpp
@@ -4,6 +4,8 @@
 
 #include "EditorStyleSet.h"
 
+#include "Algo/Transform.h"
+
 #include "Widgets/DeclarativeSyntaxSupport.h"
 #include "Widgets/Images/SImage.h"
 #include "Widgets/Layout/SScrollBorder.h"
@@ -18,6 +20,55 @@
 
 #define LOCTEXT_NAMESPACE "SourceControlChangelist"
 
+//////////////////////////////
+
+static TSharedRef<SWidget> GetSCCFileWidget(FSourceControlStateRef InFileState)
+{
+	const FSlateBrush* IconBrush = FEditorStyle::GetBrush("ContentBrowser.ColumnViewAssetIcon");
+
+	// Make icon overlays (eg, SCC and dirty status) a reasonable size in relation to the icon size (note: it is assumed this icon is square)
+	const float ICON_SCALING_FACTOR = 0.7f;
+	const float IconOverlaySize = IconBrush->ImageSize.X * ICON_SCALING_FACTOR;
+
+	return SNew(SOverlay)
+
+		// The actual icon
+		+ SOverlay::Slot()
+		[
+			SNew(SImage)
+			.Image(IconBrush)
+		]
+
+	// Source control state
+	+ SOverlay::Slot()
+		.HAlign(HAlign_Left)
+		.VAlign(VAlign_Top)
+		[
+			SNew(SBox)
+			.WidthOverride(IconOverlaySize)
+			.HeightOverride(IconOverlaySize)
+			[
+				SNew(SLayeredImage, InFileState->GetIcon())
+			]
+		];
+}
+
+struct FSCCFileDragDropOp : public FDragDropOperation
+{
+	DRAG_DROP_OPERATOR_TYPE(FSCCFileDragDropOp, FDragDropOperation);
+
+	using FDragDropOperation::Construct;
+
+	virtual TSharedPtr<SWidget> GetDefaultDecorator() const override
+	{
+		return GetSCCFileWidget(Files[0]);
+	}
+
+	TArray<FSourceControlStateRef> Files;
+};
+
+//////////////////////////////
+
 struct FChangelistTreeItem : public IChangelistTreeItem
 {
 	FChangelistTreeItem(FSourceControlChangelistStateRef InChangelistState)
@@ -544,6 +595,27 @@ public:
 		return FText::FromString(DescriptionString);
 	}
 
+protected:
+	//~ Begin STableRow Interface.
+	virtual FReply OnDrop(const FGeometry& InGeometry, const FDragDropEvent& InDragDropEvent) override
+	{
+		TSharedPtr<FSCCFileDragDropOp> Operation = InDragDropEvent.GetOperationAs<FSCCFileDragDropOp>();
+		if (Operation.IsValid())
+		{
+			FSourceControlChangelistPtr Changelist = TreeItem->ChangelistState->GetChangelist();
+			check(Changelist.IsValid());
+
+			TArray<FString> Files;
+			Algo::Transform(Operation->Files, Files, [](const FSourceControlStateRef& State) { return State->GetFilename(); });
+
+			ISourceControlProvider& SourceControlProvider = ISourceControlModule::Get().GetProvider();
+			SourceControlProvider.Execute(ISourceControlOperation::Create<FMoveToChangelist>(), Changelist, Files);
+		}
+
+		return FReply::Handled();
+	}
+	//~ End STableRow Interface.
+
 private:
 	/** The info about the widget that we are visualizing. */
 	FChangelistTreeItem* TreeItem;
@@ -556,6 +628,7 @@ public:
 		: _TreeItemToVisualize()
 	{}
 	SLATE_ARGUMENT(FChangelistTreeItemPtr, TreeItemToVisualize)
+	SLATE_EVENT(FOnDragDetected, OnDragDetected)
 	SLATE_END_ARGS()
 
 public:
@@ -568,42 +641,16 @@ public:
 	{
 		TreeItem = static_cast<FFileTreeItem*>(InArgs._TreeItemToVisualize.Get());
 
-		const FSlateBrush* IconBrush = FEditorStyle::GetBrush("ContentBrowser.ColumnViewAssetIcon");
-
-		// Make icon overlays (eg, SCC and dirty status) a reasonable size in relation to the icon size (note: it is assumed this icon is square)
-		const float ICON_SCALING_FACTOR = 0.7f;
-		const float IconOverlaySize = IconBrush->ImageSize.X * ICON_SCALING_FACTOR;
-
 		ChildSlot
 		[
 			SNew(SHorizontalBox)
 
 			// Icon
-			+SHorizontalBox::Slot()
+			+ SHorizontalBox::Slot()
 				.AutoWidth()
 				.Padding(40, 0, 4, 0)
 				[
-					SNew(SOverlay)
-				
-					// The actual icon
-					+SOverlay::Slot()
-					[
-						SNew(SImage)
-						.Image( IconBrush )
-					]
-
-					// Source control state
-					+SOverlay::Slot()
-					.HAlign(HAlign_Left)
-					.VAlign(VAlign_Top)
-					[
-						SNew(SBox)
-						.WidthOverride(IconOverlaySize)
-						.HeightOverride(IconOverlaySize)
-						[
-							SNew(SLayeredImage, TreeItem->FileState->GetIcon())
-						]
-					]
+					GetSCCFileWidget(TreeItem->FileState)
 				]
 
 			+ SHorizontalBox::Slot()
@@ -615,8 +662,11 @@ public:
 				]
 		];
 
+		auto Args = STableRow::FArguments()
+			.OnDragDetected(InArgs._OnDragDetected);
+
 		STableRow<FChangelistTreeItemPtr>::ConstructInternal(
-			STableRow::FArguments()
+			Args
 			.ShowSelection(true),
 			InOwner
 		);
@@ -627,6 +677,21 @@ public:
 		return TreeItem->GetDisplayText();
 	}
 
+protected:
+	//~ Begin STableRow Interface.
+	virtual void OnDragEnter(FGeometry const& InGeometry, FDragDropEvent const& InDragDropEvent) override
+	{
+		TSharedPtr<FDragDropOperation> DragOperation = InDragDropEvent.GetOperation();
+		DragOperation->SetCursorOverride(EMouseCursor::SlashedCircle);
+	}
+
+	virtual void OnDragLeave(FDragDropEvent const& InDragDropEvent) override
+	{
+		TSharedPtr<FDragDropOperation> DragOperation = InDragDropEvent.GetOperation();
+		DragOperation->SetCursorOverride(EMouseCursor::None);
+	}
+	//~ End STableRow Interface.
+
 private:
 	/** The info about the widget that we are visualizing. */
 	FFileTreeItem* TreeItem;
@@ -642,7 +707,8 @@ TSharedRef<ITableRow> SSourceControlChangelistsWidget::OnGenerateRow(FChangelist
 
 	case IChangelistTreeItem::File:
 		return SNew(SFileTableRow, OwnerTable)
-			.TreeItemToVisualize(InTreeItem);
+			.TreeItemToVisualize(InTreeItem)
+			.OnDragDetected(this, &SSourceControlChangelistsWidget::OnFilesDragged);
 
 	default:
 		check(false);
@@ -651,6 +717,21 @@ TSharedRef<ITableRow> SSourceControlChangelistsWidget::OnGenerateRow(FChangelist
 	return SNew(STableRow<TSharedPtr<FString>>, OwnerTable);
 }
 
+FReply SSourceControlChangelistsWidget::OnFilesDragged(const FGeometry& InGeometry, const FPointerEvent& InMouseEvent)
+{
+	if (InMouseEvent.IsMouseButtonDown(EKeys::LeftMouseButton) && !TreeView->GetSelectedItems().IsEmpty())
+	{
+		TSharedRef<FSCCFileDragDropOp> Operation = MakeShareable(new FSCCFileDragDropOp());
+
+		Algo::Transform(TreeView->GetSelectedItems(), Operation->Files, [](FChangelistTreeItemPtr InTreeItem) { check(InTreeItem->GetTreeItemType() == IChangelistTreeItem::File); return static_cast<FFileTreeItem*>(InTreeItem.Get())->FileState; });
+		Operation->Construct();
+
+		return FReply::Handled().BeginDragDrop(Operation);
+	}
+
+	return FReply::Unhandled();
+}
+
 void SSourceControlChangelistsWidget::OnGetChildren(FChangelistTreeItemPtr InParent, TArray<FChangelistTreeItemPtr>& OutChildren)
 {
 	for (auto& Child : InParent->GetChildren())
diff --git a/Engine/Source/Editor/SourceControlWindows/Private/SSourceControlChangelists.h b/Engine/Source/Editor/SourceControlWindows/Private/SSourceControlChangelists.h
--- a/Engine/Source/Editor/SourceControlWindows/Private/SSourceControlChangelists.h
+++ b/Engine/Source/Editor/SourceControlWindows/Private/SSourceControlChangelists.h
@@ -92,6 +92,8 @@ private:
 	TSharedRef<ITableRow> OnGenerateRow(FChangelistTreeItemPtr InTreeItem, const TSharedRef<STableViewBase>& OwnerTable);
 	void OnGetChildren(FChangelistTreeItemPtr InParent, TArray<FChangelistTreeItemPtr>& OutChildren);
 
+	FReply OnFilesDragged(const FGeometry& InGeometry, const FPointerEvent& InMouseEvent);
+
 	void RequestRefresh();
 	void Refresh();
 	void ClearChangelistsTree();
