Small improvements in the changelists
- Update UI for revert unchanged
- Update UI for submit CL

diff --git a/Engine/Plugins/Developer/PerforceSourceControl/Source/PerforceSourceControl/Private/PerforceSourceControlOperations.cpp b/Engine/Plugins/Developer/PerforceSourceControl/Source/PerforceSourceControl/Private/PerforceSourceControlOperations.cpp
--- a/Engine/Plugins/Developer/PerforceSourceControl/Source/PerforceSourceControl/Private/PerforceSourceControlOperations.cpp
+++ b/Engine/Plugins/Developer/PerforceSourceControl/Source/PerforceSourceControl/Private/PerforceSourceControlOperations.cpp
@@ -14,6 +14,7 @@
 #include "PerforceSourceControlModule.h"
 #include "PerforceSourceControlChangeStatusOperation.h"
 #include "SPerforceSourceControlSettings.h"
+#include "Algo/AnyOf.h"
 
 #define LOCTEXT_NAMESPACE "PerforceSourceControl"
 
@@ -452,6 +453,8 @@ bool FPerforceCheckInWorker::Execute(FPerforceSourceControlCommand& InCommand)
 				{
 					OutResults.Add(*Iter, EPerforceState::ReadOnly);
 				}
+
+				OutChangelist = ChangeList;
 			}
 		}
 
@@ -480,8 +483,16 @@ bool FPerforceCheckInWorker::Execute(FPerforceSourceControlCommand& InCommand)
 
 bool FPerforceCheckInWorker::UpdateStates() const
 {
-	// TODO : update changelist as well
-	return UpdateCachedStates(OutResults);
+	bool bUpdatedStates = UpdateCachedStates(OutResults);
+	bool bUpdatedChangelistStates = false;
+
+	if (!OutChangelist.IsDefault())
+	{
+		FPerforceSourceControlModule& PerforceSourceControl = FPerforceSourceControlModule::Get();
+		bUpdatedChangelistStates = PerforceSourceControl.GetProvider().RemoveChangelistFromCache(OutChangelist);
+	}
+
+	return (bUpdatedStates || bUpdatedChangelistStates);
 }
 
 FName FPerforceMarkForAddWorker::GetName() const
@@ -1506,6 +1517,8 @@ bool FPerforceGetPendingChangelistsWorker::Execute(FPerforceSourceControlCommand
 
 		if (Operation->ShouldUpdateFilesStates())
 		{
+			OutCLFilesStateMap.Reserve(OutChangelistsStates.Num());
+
 			for (FPerforceSourceControlChangelistState& ChangelistState : OutChangelistsStates)
 			{
 				if (!ShouldContinueProcessing())
@@ -1518,14 +1531,10 @@ bool FPerforceGetPendingChangelistsWorker::Execute(FPerforceSourceControlCommand
 				Parameters.Add(ChangelistState.Changelist.ToString());	// <changelist>
 
 				FP4RecordSet Records;
-				Connection.RunCommand(TEXT("opened"), Parameters, Records, InCommand.ResultInfo.ErrorMessages, FOnIsCancelled::CreateRaw(&InCommand, &FPerforceSourceControlCommand::IsCanceled), InCommand.bConnectionDropped);
-				InCommand.bCommandSuccessful &= InCommand.ResultInfo.ErrorMessages.IsEmpty();
+				TMap<FString, EPerforceState::Type>& OutStateMap = OutCLFilesStateMap.Emplace_GetRef();
 
-				TArray<FPerforceSourceControlState> OutStates;
-				TMap<FString, FBranchModification> BranchModifications;
-				ParseUpdateStatusResults(Records, InCommand.ResultInfo.ErrorMessages, OutStates, InCommand.ContentRoot, BranchModifications);
-
-				OutCLFilesStates.Add(OutStates);
+				InCommand.bCommandSuccessful &= Connection.RunCommand(TEXT("opened"), Parameters, Records, InCommand.ResultInfo.ErrorMessages, FOnIsCancelled::CreateRaw(&InCommand, &FPerforceSourceControlCommand::IsCanceled), InCommand.bConnectionDropped);
+				ParseOpenedResults(Records, ANSI_TO_TCHAR(Connection.P4Client.GetClient().Text()), Connection.ClientRoot, OutStateMap);
 			}
 		}
 	}
@@ -1533,7 +1542,7 @@ bool FPerforceGetPendingChangelistsWorker::Execute(FPerforceSourceControlCommand
 	if (InCommand.IsCanceled() || !InCommand.bCommandSuccessful)
 	{
 		OutChangelistsStates.Empty();
-		OutCLFilesStates.Empty();
+		OutCLFilesStateMap.Empty();
 		OutCLShelvedFilesStates.Empty();
 	}
 
@@ -1558,16 +1567,16 @@ bool FPerforceGetPendingChangelistsWorker::UpdateStates() const
 		bUpdated = true;
 
 		// Update files states for files in the changelist
-		bool bUpdateFilesStates = OutCLFilesStates.Num() == OutChangelistsStates.Num();
+		bool bUpdateFilesStates = (OutCLFilesStateMap.Num() == OutChangelistsStates.Num());
 		if (bUpdateFilesStates)
 		{
-			ChangelistState->Files.Reset(OutCLFilesStates[StatusIndex].Num());
+			ChangelistState->Files.Reset(OutCLFilesStateMap[StatusIndex].Num());
 
-			for (const FPerforceSourceControlState& Status : OutCLFilesStates[StatusIndex])
+			for(TMap<FString, EPerforceState::Type>::TConstIterator It(OutCLFilesStateMap[StatusIndex]); It; ++It)
 			{
-				TSharedRef<FPerforceSourceControlState, ESPMode::ThreadSafe> FileState = PerforceSourceControl.GetProvider().GetStateInternal(Status.LocalFilename);
+				TSharedRef<FPerforceSourceControlState, ESPMode::ThreadSafe> FileState = PerforceSourceControl.GetProvider().GetStateInternal(It.Key());
 				auto History = MoveTemp(FileState->History);
-				*FileState = Status;
+				FileState->SetState(It.Value());
 				FileState->History = MoveTemp(History);
 				FileState->TimeStamp = Now;
 				
@@ -1886,22 +1895,35 @@ bool FPerforceRevertUnchangedWorker::Execute(class FPerforceSourceControlCommand
 		if (InCommand.Files.Num() > 0)
 		{
 			Parameters.Append(InCommand.Files);
-		}
-		else
-		{
-			Parameters.Add(TEXT("..."));
-		}		
+		}	
 
 		FP4RecordSet Records;
 		InCommand.bCommandSuccessful = Connection.RunCommand(TEXT("revert"), Parameters, Records, InCommand.ResultInfo.ErrorMessages, FOnIsCancelled::CreateRaw(&InCommand, &FPerforceSourceControlCommand::IsCanceled), InCommand.bConnectionDropped);
 		ParseRecordSetForState(Records, OutResults);
+		ChangelistToUpdate = InCommand.Changelist;
 	}
 	return InCommand.bCommandSuccessful;
 }
 
 bool FPerforceRevertUnchangedWorker::UpdateStates() const
 {
-	return UpdateCachedStates(OutResults);
+	bool bUpdatedStates = UpdateCachedStates(OutResults);
+	bool bUpdatedChangelistState = false;
+
+	if (OutResults.Num() > 0)
+	{
+		FPerforceSourceControlModule& PerforceSourceControl = FPerforceSourceControlModule::Get();
+		TSharedRef<FPerforceSourceControlChangelistState, ESPMode::ThreadSafe> ChangelistState = PerforceSourceControl.GetProvider().GetStateInternal(ChangelistToUpdate);
+
+		bUpdatedChangelistState = ChangelistState->Files.RemoveAll([this](FSourceControlStateRef& State) -> bool
+			{
+				return Algo::AnyOf(OutResults, [&State](auto& OutResult) {
+					return State->GetFilename() == OutResult.Key;
+					});
+			}) > 0;
+	}
+
+	return bUpdatedStates || bUpdatedChangelistState;
 }
 
 FName FPerforceReopenWorker::GetName() const
diff --git a/Engine/Plugins/Developer/PerforceSourceControl/Source/PerforceSourceControl/Private/PerforceSourceControlOperations.h b/Engine/Plugins/Developer/PerforceSourceControl/Source/PerforceSourceControl/Private/PerforceSourceControlOperations.h
--- a/Engine/Plugins/Developer/PerforceSourceControl/Source/PerforceSourceControl/Private/PerforceSourceControlOperations.h
+++ b/Engine/Plugins/Developer/PerforceSourceControl/Source/PerforceSourceControl/Private/PerforceSourceControlOperations.h
@@ -47,7 +47,7 @@ public:
 	TMap<FString, EPerforceState::Type> OutResults;
 
 	/** Changelist we submitted */
-	int32 OutChangelistNumber;
+	FPerforceSourceControlChangelist OutChangelist;
 };
 
 class FPerforceMarkForAddWorker : public IPerforceSourceControlWorker
@@ -152,7 +152,7 @@ public:
 public:
 	/** Temporary states for results */
 	TArray<FPerforceSourceControlChangelistState> OutChangelistsStates;
-	TArray<TArray<FPerforceSourceControlState>> OutCLFilesStates;
+	TArray<TMap<FString, EPerforceState::Type>> OutCLFilesStateMap;
 	TArray<TArray<FPerforceSourceControlState>> OutCLShelvedFilesStates;
 };
 
@@ -246,6 +246,9 @@ public:
 protected:
 	/** Map of filenames to perforce state */
 	TMap<FString, EPerforceState::Type> OutResults;
+
+	/** Changelist to be udpated */
+	FPerforceSourceControlChangelist ChangelistToUpdate;
 };
 
 class FPerforceReopenWorker : public IPerforceSourceControlWorker
