Added support for shelved files in changelist, including shelve, unshelve and delete shelve operations.
Added update state for reopen operation as well.

diff --git a/Engine/Plugins/Developer/PerforceSourceControl/Source/PerforceSourceControl/Private/PerforceSourceControlChangelist.h b/Engine/Plugins/Developer/PerforceSourceControl/Source/PerforceSourceControl/Private/PerforceSourceControlChangelist.h
--- a/Engine/Plugins/Developer/PerforceSourceControl/Source/PerforceSourceControl/Private/PerforceSourceControlChangelist.h
+++ b/Engine/Plugins/Developer/PerforceSourceControl/Source/PerforceSourceControl/Private/PerforceSourceControlChangelist.h
@@ -18,6 +18,7 @@ public:
 
 	FPerforceSourceControlChangelist(const FPerforceSourceControlChangelist& InOther)
 		: ChangelistNumber(InOther.ChangelistNumber)
+		, bInitialized(InOther.bInitialized)
 	{
 	}
 
@@ -57,8 +58,15 @@ public:
 		return ChangelistNumber != DefaultChangelist.ChangelistNumber ? FString::FromInt(ChangelistNumber) : TEXT("default");
 	}
 
+	int32 ToInt() const
+	{
+		return ChangelistNumber;
+	}
+
 public:
 	static const FPerforceSourceControlChangelist DefaultChangelist;
+
+private:
 	int32 ChangelistNumber;
 	bool bInitialized;
 };
diff --git a/Engine/Plugins/Developer/PerforceSourceControl/Source/PerforceSourceControl/Private/PerforceSourceControlModule.cpp b/Engine/Plugins/Developer/PerforceSourceControl/Source/PerforceSourceControl/Private/PerforceSourceControlModule.cpp
--- a/Engine/Plugins/Developer/PerforceSourceControl/Source/PerforceSourceControl/Private/PerforceSourceControlModule.cpp
+++ b/Engine/Plugins/Developer/PerforceSourceControl/Source/PerforceSourceControl/Private/PerforceSourceControlModule.cpp
@@ -34,6 +34,10 @@ void FPerforceSourceControlModule::StartupModule()
 	PerforceSourceControlProvider.RegisterWorker( "EditChangelist", FGetPerforceSourceControlWorker::CreateStatic( &CreateWorker<FPerforceEditChangelistWorker> ) );
 	PerforceSourceControlProvider.RegisterWorker( "RevertUnchanged", FGetPerforceSourceControlWorker::CreateStatic( &CreateWorker<FPerforceRevertUnchangedWorker> ) );
 	PerforceSourceControlProvider.RegisterWorker( "MoveToChangelist", FGetPerforceSourceControlWorker::CreateStatic( &CreateWorker<FPerforceReopenWorker>) );
+	PerforceSourceControlProvider.RegisterWorker( "Shelve", FGetPerforceSourceControlWorker::CreateStatic( &CreateWorker<FPerforceShelveWorker>) );
+	PerforceSourceControlProvider.RegisterWorker( "Unshelve", FGetPerforceSourceControlWorker::CreateStatic( &CreateWorker<FPerforceUnshelveWorker>) );
+	PerforceSourceControlProvider.RegisterWorker( "DeleteShelved", FGetPerforceSourceControlWorker::CreateStatic( &CreateWorker<FPerforceDeleteShelveWorker>) );
+
 
 	// load our settings
 	PerforceSourceControlSettings.LoadSettings();
diff --git a/Engine/Plugins/Developer/PerforceSourceControl/Source/PerforceSourceControl/Private/PerforceSourceControlOperations.cpp b/Engine/Plugins/Developer/PerforceSourceControl/Source/PerforceSourceControl/Private/PerforceSourceControlOperations.cpp
--- a/Engine/Plugins/Developer/PerforceSourceControl/Source/PerforceSourceControl/Private/PerforceSourceControlOperations.cpp
+++ b/Engine/Plugins/Developer/PerforceSourceControl/Source/PerforceSourceControl/Private/PerforceSourceControlOperations.cpp
@@ -15,6 +15,8 @@
 #include "PerforceSourceControlChangeStatusOperation.h"
 #include "SPerforceSourceControlSettings.h"
 #include "Algo/AnyOf.h"
+#include "Algo/IndexOf.h"
+#include "Algo/Find.h"
 
 #define LOCTEXT_NAMESPACE "PerforceSourceControl"
 
@@ -861,6 +863,7 @@ static void ParseUpdateStatusResults(const FP4RecordSet& InRecords, const TArray
 		const FP4Record& ClientRecord = InRecords[Index];
 		FString FileName = ClientRecord(TEXT("clientFile"));
 		FString DepotFileName = ClientRecord(TEXT("depotFile"));
+		FString Changelist = ClientRecord(TEXT("change"));
 		FString HeadRev  = ClientRecord(TEXT("headRev"));
 		FString HaveRev  = ClientRecord(TEXT("haveRev"));
 		FString OtherOpen = ClientRecord(TEXT("otherOpen"));
@@ -942,6 +945,15 @@ static void ParseUpdateStatusResults(const FP4RecordSet& InRecords, const TArray
 			State.State = EPerforceState::NotInDepot;
 		}
 
+		if (Changelist.Len() > 0 && Changelist != TEXT("default"))
+		{
+			State.Changelist = FPerforceSourceControlChangelist(FCString::Atoi(*Changelist));
+		}
+		else
+		{
+			State.Changelist = FPerforceSourceControlChangelist::DefaultChangelist;
+		}
+
 		State.HeadBranch = TEXT("*CurrentBranch");
 		State.HeadAction = HeadAction;
 		State.HeadModTime = FCString::Atoi64(*ClientRecord(TEXT("headModTime")));
@@ -1072,14 +1084,13 @@ static void ParseUpdateStatusResults(const FP4RecordSet& InRecords, const TArray
 	}
 }
 
-static void ParseOpenedResults(const FP4RecordSet& InRecords, const FString& ClientName, const FString& ClientRoot, TMap<FString, EPerforceState::Type>& OutResults)
+static void ParseOpenedResults(const FP4RecordSet& InRecords, const FString& ClientName, const FString& ClientRoot, TArray<FPerforceSourceControlState>& OutResults)
 {
 	// Iterate over each record found as a result of the command, parsing it for relevant information
 	for (int32 Index = 0; Index < InRecords.Num(); ++Index)
 	{
 		const FP4Record& ClientRecord = InRecords[Index];
 		FString ClientFileName = ClientRecord(TEXT("clientFile"));
-		FString Action = ClientRecord(TEXT("action"));
 
 		check(ClientFileName.Len() > 0);
 
@@ -1108,19 +1119,114 @@ static void ParseOpenedResults(const FP4RecordSet& InRecords, const FString& Cli
 			continue;
 		}
 
+		// Fill-in with information we got from the opened command, namely:
+		// depotFile, rev, haveRev, action, change, type, user, client
+		// Note: haveRev works, but we don't have the depot revision, so we might as well not write anything
+		FPerforceSourceControlState& OutState = OutResults.Emplace_GetRef(FullPath);
+		OutState.DepotFilename = ClientRecord(TEXT("depotFile"));
+
+		FString Action = ClientRecord(TEXT("action"));
 		if (Action.Len() > 0)
 		{
 			if(IsAddAction(Action))
 			{
-				OutResults.Add(FullPath, EPerforceState::OpenForAdd);
+				OutState.State = EPerforceState::OpenForAdd;
 			}
 			else if(Action == TEXT("edit"))
 			{
-				OutResults.Add(FullPath, EPerforceState::CheckedOut);
+				OutState.State = EPerforceState::CheckedOut;
 			}
 			else if(IsDeleteAction(Action))
 			{
-				OutResults.Add(FullPath, EPerforceState::MarkedForDelete);
+				OutState.State = EPerforceState::MarkedForDelete;
+			}
+		}
+
+		FString Changelist = ClientRecord(TEXT("change"));
+		if (Changelist.Len() > 0 && Changelist != TEXT("default"))
+		{
+			OutState.Changelist = FPerforceSourceControlChangelist(FCString::Atoi(*Changelist));
+		}
+		else
+		{
+			OutState.Changelist = FPerforceSourceControlChangelist::DefaultChangelist;
+		}
+
+		FString Type = ClientRecord(TEXT("type"));
+		if (Type.Len() > 0)
+		{
+			OutState.bBinary = Type.Contains(TEXT("binary"));
+			OutState.bExclusiveCheckout = Type.Contains(TEXT("+l"));
+		}
+	}
+}
+
+static void ParseOpenedResults(const FP4RecordSet& InRecords, const FString& ClientName, const FString& ClientRoot, TMap<FString, EPerforceState::Type>& OutResults)
+{
+	TArray<FPerforceSourceControlState> TemporaryStates;
+	ParseOpenedResults(InRecords, ClientName, ClientRoot, TemporaryStates);
+
+	for (const FPerforceSourceControlState& FileState : TemporaryStates)
+	{
+		if (FileState.State != EPerforceState::DontCare)
+			OutResults.Add(FileState.LocalFilename, FileState.State);
+	}
+}
+
+static void ParseShelvedResults(const FP4RecordSet& InRecords, TMap<FString, EPerforceState::Type>& OutResults)
+{
+	// Iterate over each record found as a result of the command, parsing it for relevant information
+	for (int32 Index = 0; Index < InRecords.Num(); ++Index)
+	{
+		const FP4Record& Record = InRecords[Index];
+		FString DepotFileName = Record(TEXT("depotFile"));
+		FString Action = Record(TEXT("action"));
+
+		if (Action.Len() > 0 && DepotFileName.Len() > 0)
+		{
+			if (IsAddAction(Action))
+			{
+				OutResults.Add(DepotFileName, EPerforceState::OpenForAdd);
+			}
+			else if (Action == TEXT("edit"))
+			{
+				OutResults.Add(DepotFileName, EPerforceState::CheckedOut);
+			}
+			else if (IsDeleteAction(Action))
+			{
+				OutResults.Add(DepotFileName, EPerforceState::MarkedForDelete);
+			}
+		}
+	}
+}
+
+static void ParseShelvedChangelistResults(const FP4RecordSet& InRecords, TMap<FString, EPerforceState::Type>& OutResults)
+{
+	// Describe returns only one record.
+	check(InRecords.Num() == 1);
+	const FP4Record& Record = InRecords[0];
+
+	for (int32 FileIndex = 0;; ++FileIndex)
+	{
+		FString DepotFileName = Record(FString::Printf(TEXT("depotFile%d"), FileIndex));
+		FString Action = Record(FString::Printf(TEXT("action%d"), FileIndex));
+
+		if (DepotFileName.Len() == 0)
+			break;
+
+		if (Action.Len() > 0)
+		{
+			if (IsAddAction(Action))
+			{
+				OutResults.Add(DepotFileName, EPerforceState::OpenForAdd);
+			}
+			else if (Action == TEXT("edit"))
+			{
+				OutResults.Add(DepotFileName, EPerforceState::CheckedOut);
+			}
+			else if (IsDeleteAction(Action))
+			{
+				OutResults.Add(DepotFileName, EPerforceState::MarkedForDelete);
 			}
 		}
 	}
@@ -1520,6 +1626,24 @@ FName FPerforceGetPendingChangelistsWorker::GetName() const
 	return "GetPendingChangelists";
 }
 
+static bool GetOpenedFilesInChangelist(FPerforceConnection& Connection, FPerforceSourceControlCommand& InCommand, const FPerforceSourceControlChangelist& Changelist, TArray<FPerforceSourceControlState>& FilesStates)
+{
+	TArray<FString> Parameters;
+	Parameters.Add(TEXT("-c"));	// -c	Changelist
+	Parameters.Add(Changelist.ToString());	// <changelist>
+
+	FP4RecordSet Records;
+	Connection.RunCommand(TEXT("opened"), Parameters, Records, InCommand.ResultInfo.ErrorMessages, FOnIsCancelled::CreateRaw(&InCommand, &FPerforceSourceControlCommand::IsCanceled), InCommand.bConnectionDropped);
+	InCommand.bCommandSuccessful &= (InCommand.ResultInfo.ErrorMessages.Num() == 0);
+
+	if (InCommand.bCommandSuccessful)
+	{
+		ParseOpenedResults(Records, ANSI_TO_TCHAR(Connection.P4Client.GetClient().Text()), Connection.ClientRoot, FilesStates);
+	}
+
+	return InCommand.bCommandSuccessful;
+}
+
 bool FPerforceGetPendingChangelistsWorker::Execute(FPerforceSourceControlCommand& InCommand)
 {
 	FScopedPerforceConnection ScopedConnection(InCommand);
@@ -1557,7 +1681,22 @@ bool FPerforceGetPendingChangelistsWorker::Execute(FPerforceSourceControlCommand
 
 		if (Operation->ShouldUpdateFilesStates())
 		{
-			OutCLFilesStateMap.Reserve(OutChangelistsStates.Num());
+			OutCLFilesStates.Reserve(OutChangelistsStates.Num());
+
+			for (FPerforceSourceControlChangelistState& ChangelistState : OutChangelistsStates)
+			{
+				if (!ShouldContinueProcessing())
+				{
+					break;
+				}
+
+				GetOpenedFilesInChangelist(Connection, InCommand, ChangelistState.Changelist, OutCLFilesStates.Emplace_GetRef());
+			}
+		}
+
+		if(Operation->ShouldUpdateShelvedFilesStates())
+		{
+			OutCLShelvedFilesStates.Reserve(OutChangelistsStates.Num());
 
 			for (FPerforceSourceControlChangelistState& ChangelistState : OutChangelistsStates)
 			{
@@ -1566,16 +1705,26 @@ bool FPerforceGetPendingChangelistsWorker::Execute(FPerforceSourceControlCommand
 					break;
 				}
 
+				if (!ChangelistState.bHasShelvedFiles)
+				{
+					OutCLShelvedFilesStates.Emplace(); // Add empty list
+					continue;
+				}
+
 				TArray<FString> Parameters;
-				Parameters.Add(TEXT("-c"));								// -c			Changelist
-				Parameters.Add(ChangelistState.Changelist.ToString());	// <changelist>
+				Parameters.Add(TEXT("-s"));
+				Parameters.Add(TEXT("-S"));
+				Parameters.Add(ChangelistState.Changelist.ToString());
 
 				FP4RecordSet Records;
-				TMap<FString, EPerforceState::Type>& OutStateMap = OutCLFilesStateMap.Emplace_GetRef();
-
-				Connection.RunCommand(TEXT("opened"), Parameters, Records, InCommand.ResultInfo.ErrorMessages, FOnIsCancelled::CreateRaw(&InCommand, &FPerforceSourceControlCommand::IsCanceled), InCommand.bConnectionDropped);
+				Connection.RunCommand(TEXT("describe"), Parameters, Records, InCommand.ResultInfo.ErrorMessages, FOnIsCancelled::CreateRaw(&InCommand, &FPerforceSourceControlCommand::IsCanceled), InCommand.bConnectionDropped);
 				InCommand.bCommandSuccessful &= (InCommand.ResultInfo.ErrorMessages.Num() == 0);
-				ParseOpenedResults(Records, ANSI_TO_TCHAR(Connection.P4Client.GetClient().Text()), Connection.ClientRoot, OutStateMap);
+
+				if (InCommand.bCommandSuccessful)
+				{
+					TMap<FString, EPerforceState::Type>& OutShelvedStateMap = OutCLShelvedFilesStates.Emplace_GetRef();
+					ParseShelvedChangelistResults(Records, OutShelvedStateMap);
+				}
 			}
 		}
 	}
@@ -1583,13 +1732,47 @@ bool FPerforceGetPendingChangelistsWorker::Execute(FPerforceSourceControlCommand
 	if (InCommand.IsCanceled() || !InCommand.bCommandSuccessful)
 	{
 		OutChangelistsStates.Empty();
-		OutCLFilesStateMap.Empty();
+		OutCLFilesStates.Empty();
 		OutCLShelvedFilesStates.Empty();
 	}
 
 	return InCommand.bCommandSuccessful;
 }
 
+static bool AddShelvedFilesToChangelist(const TMap<FString, EPerforceState::Type>& FilesToAdd, TSharedRef<FPerforceSourceControlChangelistState, ESPMode::ThreadSafe>& ChangelistState, const FDateTime* TimeStamp = nullptr)
+{
+	FPerforceSourceControlModule& PerforceSourceControl = FPerforceSourceControlModule::Get();
+	const FDateTime Now = (TimeStamp ? *TimeStamp : FDateTime::Now());
+
+	for (TMap<FString, EPerforceState::Type>::TConstIterator It(FilesToAdd); It; ++It)
+	{
+		int32 Index = Algo::IndexOfByPredicate(ChangelistState->ShelvedFiles, [&It](const FSourceControlStateRef& ShelvedFile) {
+				return ShelvedFile->GetFilename() == It.Key();
+			});
+
+		if (Index < 0)
+		{
+			// Create new entry
+			Index = ChangelistState->ShelvedFiles.Num();
+			ChangelistState->ShelvedFiles.Emplace(MakeShareable(new FPerforceSourceControlState(It.Key())));
+		}
+
+		TSharedRef<FPerforceSourceControlState, ESPMode::ThreadSafe> FileState = StaticCastSharedRef<FPerforceSourceControlState>(ChangelistState->ShelvedFiles[Index]);
+
+ 		FileState->SetState(It.Value());
+ 		FileState->TimeStamp = Now;
+	}
+
+	return FilesToAdd.Num() > 0;
+}
+
+static bool AddShelvedFilesToChangelist(const TMap<FString, EPerforceState::Type>& FilesToAdd, const FPerforceSourceControlChangelist& Changelist)
+{
+	FPerforceSourceControlModule& PerforceSourceControl = FPerforceSourceControlModule::Get();
+	TSharedRef<FPerforceSourceControlChangelistState, ESPMode::ThreadSafe> ChangelistState = PerforceSourceControl.GetProvider().GetStateInternal(Changelist);
+	return AddShelvedFilesToChangelist(FilesToAdd, ChangelistState);
+}
+
 bool FPerforceGetPendingChangelistsWorker::UpdateStates() const
 {
 	bool bUpdated = false;
@@ -1608,22 +1791,25 @@ bool FPerforceGetPendingChangelistsWorker::UpdateStates() const
 		bUpdated = true;
 
 		// Update files states for files in the changelist
-		bool bUpdateFilesStates = (OutCLFilesStateMap.Num() == OutChangelistsStates.Num());
+		bool bUpdateFilesStates = (OutCLFilesStates.Num() == OutChangelistsStates.Num());
 		if (bUpdateFilesStates)
 		{
-			ChangelistState->Files.Reset(OutCLFilesStateMap[StatusIndex].Num());
-
-			for(TMap<FString, EPerforceState::Type>::TConstIterator It(OutCLFilesStateMap[StatusIndex]); It; ++It)
+			ChangelistState->Files.Reset(OutCLFilesStates[StatusIndex].Num());
+			for (const auto& FileState : OutCLFilesStates[StatusIndex])
 			{
-				TSharedRef<FPerforceSourceControlState, ESPMode::ThreadSafe> FileState = PerforceSourceControl.GetProvider().GetStateInternal(It.Key());
-				auto History = MoveTemp(FileState->History);
-				FileState->SetState(It.Value());
-				FileState->History = MoveTemp(History);
-				FileState->TimeStamp = Now;
-				
-				ChangelistState->Files.Add(FileState);
+				TSharedRef<FPerforceSourceControlState, ESPMode::ThreadSafe> CachedFileState = PerforceSourceControl.GetProvider().GetStateInternal(FileState.LocalFilename);
+				CachedFileState->Update(FileState, &Now);
+				ChangelistState->Files.AddUnique(CachedFileState);
 			}
 		}
+
+		// Update shelved files in the the changelist
+		bool bUpdateShelvedFiles = (OutCLShelvedFilesStates.Num() == OutChangelistsStates.Num());
+		if(bUpdateShelvedFiles)
+		{
+			ChangelistState->ShelvedFiles.Reset(OutCLShelvedFilesStates[StatusIndex].Num());
+			AddShelvedFilesToChangelist(OutCLShelvedFilesStates[StatusIndex], ChangelistState, &Now);
+		}
 	}
 
 	return bUpdated;
@@ -1797,7 +1983,7 @@ bool FPerforceNewChangelistWorker::Execute(class FPerforceSourceControlCommand&
 
 		if (InCommand.bCommandSuccessful)
 		{
-			NewChangelist.ChangelistNumber = ChangeList;
+			NewChangelist = FPerforceSourceControlChangelist(ChangeList);
 			NewChangelistState.Changelist = NewChangelist;
 			NewChangelistState.Description = Operation->GetDescription().ToString();
 			NewChangelistState.bHasShelvedFiles = false;
@@ -1892,9 +2078,9 @@ bool FPerforceEditChangelistWorker::Execute(class FPerforceSourceControlCommand&
 		check(InCommand.Operation->GetName() == GetName());
 		TSharedRef<FEditChangelist, ESPMode::ThreadSafe> Operation = StaticCastSharedRef<FEditChangelist>(InCommand.Operation);
 
-		int32 ChangelistNumber = Connection.EditPendingChangelist(Operation->GetDescription(), InCommand.Changelist.ChangelistNumber, FOnIsCancelled::CreateRaw(&InCommand, &FPerforceSourceControlCommand::IsCanceled), InCommand.ResultInfo.ErrorMessages);
+		int32 ChangelistNumber = Connection.EditPendingChangelist(Operation->GetDescription(), InCommand.Changelist.ToInt(), FOnIsCancelled::CreateRaw(&InCommand, &FPerforceSourceControlCommand::IsCanceled), InCommand.ResultInfo.ErrorMessages);
 
-		InCommand.bCommandSuccessful = (ChangelistNumber == InCommand.Changelist.ChangelistNumber);
+		InCommand.bCommandSuccessful = (ChangelistNumber == InCommand.Changelist.ToInt());
 
 		if (InCommand.bCommandSuccessful)
 		{
@@ -1963,7 +2149,9 @@ bool FPerforceReopenWorker::Execute(FPerforceSourceControlCommand& InCommand)
 	FScopedPerforceConnection ScopedConnection(InCommand);
 	if (!InCommand.IsCanceled() && ScopedConnection.IsValid())
 	{
-		InCommand.bCommandSuccessful = RunReopenCommand(InCommand, InCommand.Files, InCommand.Changelist);
+		ReopenedFiles.Reset(InCommand.Files.Num());
+		InCommand.bCommandSuccessful = RunReopenCommand(InCommand, InCommand.Files, InCommand.Changelist, &ReopenedFiles);
+		DestinationChangelist = InCommand.Changelist;
 	}
 
 	return InCommand.bCommandSuccessful;
@@ -1971,8 +2159,201 @@ bool FPerforceReopenWorker::Execute(FPerforceSourceControlCommand& InCommand)
 
 bool FPerforceReopenWorker::UpdateStates() const
 {
-	// TODO : update changelists states (source & destination CLs)
-	return true;
+	const FDateTime Now = FDateTime::Now();
+	FPerforceSourceControlModule& PerforceSourceControl = FPerforceSourceControlModule::Get();
+	TSharedRef<FPerforceSourceControlChangelistState, ESPMode::ThreadSafe> DestinationChangelistState = PerforceSourceControl.GetProvider().GetStateInternal(DestinationChangelist);
+
+	// 3 things to do here:
+	for (const FString& ReopenedFile : ReopenedFiles)
+	{
+		TSharedRef<FPerforceSourceControlState, ESPMode::ThreadSafe> FileState = PerforceSourceControl.GetProvider().GetStateInternal(ReopenedFile);
+
+		// 1- Remove these files from their previous changelist
+		TSharedRef<FPerforceSourceControlChangelistState, ESPMode::ThreadSafe> PreviousChangelist = PerforceSourceControl.GetProvider().GetStateInternal(FileState->Changelist);
+		PreviousChangelist->Files.Remove(FileState);
+
+		// 2- Add to the new changelist
+		DestinationChangelistState->Files.Add(FileState);
+
+		// 3- Update changelist in file state
+		FileState->Changelist = DestinationChangelist;
+		FileState->TimeStamp = Now;
+	}
+	
+	return ReopenedFiles.Num() > 0;
+}
+
+
+FName FPerforceShelveWorker::GetName() const
+{
+	return "Shelve";
+}
+
+bool FPerforceShelveWorker::Execute(class FPerforceSourceControlCommand& InCommand)
+{
+	FScopedPerforceConnection ScopedConnection(InCommand);
+	if (!InCommand.IsCanceled() && ScopedConnection.IsValid())
+	{
+		FPerforceConnection& Connection = ScopedConnection.GetConnection();
+		TArray<FString> Parameters;
+		Parameters.Add(TEXT("-c"));
+		Parameters.Add(InCommand.Changelist.ToString());
+		Parameters.Add(TEXT("-f")); // force
+
+		if (InCommand.Files.Num() > 0)
+		{
+			Parameters.Append(InCommand.Files);
+		}
+
+		FP4RecordSet Records;
+		InCommand.bCommandSuccessful = Connection.RunCommand(TEXT("shelve"), Parameters, Records, InCommand.ResultInfo.ErrorMessages, FOnIsCancelled::CreateRaw(&InCommand, &FPerforceSourceControlCommand::IsCanceled), InCommand.bConnectionDropped);
+		
+		if (InCommand.bCommandSuccessful)
+		{
+			ChangelistToUpdate = InCommand.Changelist;
+			ParseShelvedResults(Records, OutResults);
+		}
+	}
+
+	return InCommand.bCommandSuccessful;
+}
+
+bool FPerforceShelveWorker::UpdateStates() const
+{
+	return (OutResults.Num() > 0 && AddShelvedFilesToChangelist(OutResults, ChangelistToUpdate));
+}
+
+
+FName FPerforceDeleteShelveWorker::GetName() const
+{
+	return "DeleteShelved";
+}
+
+bool FPerforceDeleteShelveWorker::Execute(class FPerforceSourceControlCommand& InCommand)
+{
+	FScopedPerforceConnection ScopedConnection(InCommand);
+	if (!InCommand.IsCanceled() && ScopedConnection.IsValid())
+	{
+		FPerforceConnection& Connection = ScopedConnection.GetConnection();
+		TArray<FString> Parameters;
+		Parameters.Add(TEXT("-d")); // -d is delete
+		Parameters.Add(TEXT("-c"));
+		Parameters.Add(InCommand.Changelist.ToString());
+
+		if (InCommand.Files.Num() > 0)
+		{
+			Parameters.Append(InCommand.Files);
+		}
+
+		FP4RecordSet Records;
+		Connection.RunCommand(TEXT("shelve"), Parameters, Records, InCommand.ResultInfo.ErrorMessages, FOnIsCancelled::CreateRaw(&InCommand, &FPerforceSourceControlCommand::IsCanceled), InCommand.bConnectionDropped);
+		InCommand.bCommandSuccessful = (InCommand.ResultInfo.ErrorMessages.Num() == 0);
+
+		if (InCommand.bCommandSuccessful)
+		{
+			ChangelistToUpdate = InCommand.Changelist;
+			FilesToRemove = InCommand.Files;
+		}
+	}
+
+	return InCommand.bCommandSuccessful;
+}
+
+bool FPerforceDeleteShelveWorker::UpdateStates() const
+{
+	if (ChangelistToUpdate.IsInitialized())
+	{
+		FPerforceSourceControlModule& PerforceSourceControl = FPerforceSourceControlModule::Get();
+		TSharedRef<FPerforceSourceControlChangelistState, ESPMode::ThreadSafe> ChangelistState = PerforceSourceControl.GetProvider().GetStateInternal(ChangelistToUpdate);
+
+		if (FilesToRemove.Num() > 0)
+		{
+			return ChangelistState->ShelvedFiles.RemoveAll([this](FSourceControlStateRef& State) -> bool
+				{
+					return Algo::AnyOf(FilesToRemove, [&State](auto& File) {
+						return State->GetFilename() == File;
+					});
+				}) > 0;
+		}
+		else
+		{
+			bool bHadShelvedFiles = (ChangelistState->ShelvedFiles.Num() > 0);
+			ChangelistState->ShelvedFiles.Reset();
+			return bHadShelvedFiles;
+		}		
+	}
+	else
+	{
+		return false;
+	}
+}
+
+
+FName FPerforceUnshelveWorker::GetName() const
+{
+	return "Unshelve";
+}
+
+bool FPerforceUnshelveWorker::Execute(class FPerforceSourceControlCommand& InCommand)
+{
+	FScopedPerforceConnection ScopedConnection(InCommand);
+	if (!InCommand.IsCanceled() && ScopedConnection.IsValid())
+	{
+		FPerforceConnection& Connection = ScopedConnection.GetConnection();
+		TArray<FString> Parameters;
+
+		Parameters.Add(TEXT("-s")); // unshelve from source changelist
+		Parameters.Add(InCommand.Changelist.ToString()); // current changelist
+		Parameters.Add(TEXT("-f")); // force overwriting of writeable but unopened files
+		Parameters.Add(TEXT("-c")); // unshelve to target changelist
+		Parameters.Add(InCommand.Changelist.ToString()); // current changelist
+
+		if (InCommand.Files.Num() > 0)
+		{
+			Parameters.Append(InCommand.Files);
+		}
+
+		FP4RecordSet Records;
+		// Note: unshelve can succeed partially.
+		InCommand.bCommandSuccessful = Connection.RunCommand(TEXT("unshelve"), Parameters, Records, InCommand.ResultInfo.ErrorMessages, FOnIsCancelled::CreateRaw(&InCommand, &FPerforceSourceControlCommand::IsCanceled), InCommand.bConnectionDropped);
+		
+		if (InCommand.bCommandSuccessful && Records.Num() > 0)
+		{
+			// At this point, the records contain the list of files from the depot that were unshelved
+			// however they contain only the depot file equivalency; considering that some files might not be in the cache yet,
+			// it is simpler to do a full update of the changelist files.
+			ChangelistToUpdate = InCommand.Changelist;
+			GetOpenedFilesInChangelist(Connection, InCommand, ChangelistToUpdate, ChangelistFilesStates);
+		}
+	}
+
+	return InCommand.bCommandSuccessful;
+}
+
+bool FPerforceUnshelveWorker::UpdateStates() const
+{
+	if (ChangelistToUpdate.IsInitialized() && ChangelistFilesStates.Num() > 0)
+	{
+		const FDateTime Now = FDateTime::Now();
+
+		FPerforceSourceControlModule& PerforceSourceControl = FPerforceSourceControlModule::Get();
+		TSharedRef<FPerforceSourceControlChangelistState, ESPMode::ThreadSafe> ChangelistState = PerforceSourceControl.GetProvider().GetStateInternal(ChangelistToUpdate);
+
+		ChangelistState->Files.Reset(ChangelistFilesStates.Num());
+		for (const auto& FileState : ChangelistFilesStates)
+		{
+			TSharedRef<FPerforceSourceControlState, ESPMode::ThreadSafe> CachedFileState = PerforceSourceControl.GetProvider().GetStateInternal(FileState.LocalFilename);
+			CachedFileState->Update(FileState, &Now);
+
+			ChangelistState->Files.AddUnique(CachedFileState);
+		}
+
+		return true;
+	}
+	else
+	{
+		return false;
+	}
 }
 
 #undef LOCTEXT_NAMESPACE
diff --git a/Engine/Plugins/Developer/PerforceSourceControl/Source/PerforceSourceControl/Private/PerforceSourceControlOperations.h b/Engine/Plugins/Developer/PerforceSourceControl/Source/PerforceSourceControl/Private/PerforceSourceControlOperations.h
--- a/Engine/Plugins/Developer/PerforceSourceControl/Source/PerforceSourceControl/Private/PerforceSourceControlOperations.h
+++ b/Engine/Plugins/Developer/PerforceSourceControl/Source/PerforceSourceControl/Private/PerforceSourceControlOperations.h
@@ -155,8 +155,8 @@ public:
 public:
 	/** Temporary states for results */
 	TArray<FPerforceSourceControlChangelistState> OutChangelistsStates;
-	TArray<TMap<FString, EPerforceState::Type>> OutCLFilesStateMap;
-	TArray<TArray<FPerforceSourceControlState>> OutCLShelvedFilesStates;
+	TArray<TArray<FPerforceSourceControlState>> OutCLFilesStates;
+	TArray<TMap<FString, EPerforceState::Type>> OutCLShelvedFilesStates;
 };
 
 class FPerforceCopyWorker : public IPerforceSourceControlWorker
@@ -250,7 +250,7 @@ protected:
 	/** Map of filenames to perforce state */
 	TMap<FString, EPerforceState::Type> OutResults;
 
-	/** Changelist to be udpated */
+	/** Changelist to be updated */
 	FPerforceSourceControlChangelist ChangelistToUpdate;
 };
 
@@ -262,5 +262,62 @@ public:
 	virtual FName GetName() const override;
 	virtual bool Execute(class FPerforceSourceControlCommand& InCommand) override;
 	virtual bool UpdateStates() const override;
+
+protected:
+	/** Reopened files */
+	TArray<FString> ReopenedFiles;
+	
+	/** Destination changelist */
+	FPerforceSourceControlChangelist DestinationChangelist;
+};
+
+class FPerforceShelveWorker : public IPerforceSourceControlWorker
+{
+public:
+	virtual ~FPerforceShelveWorker() {}
+	// IPerforceSourceControlWorker interface
+	virtual FName GetName() const override;
+	virtual bool Execute(class FPerforceSourceControlCommand& InCommand) override;
+	virtual bool UpdateStates() const override;
+
+protected:	
+	/** Map of filenames to perforce state */
+	TMap<FString, EPerforceState::Type> OutResults;
+
+	/** Changelist to be updated */
+	FPerforceSourceControlChangelist ChangelistToUpdate;
 };
 
+class FPerforceDeleteShelveWorker : public IPerforceSourceControlWorker
+{
+public:
+	virtual ~FPerforceDeleteShelveWorker() {}
+	// IPerforceSourceControlWorker interface
+	virtual FName GetName() const override;
+	virtual bool Execute(class FPerforceSourceControlCommand& InCommand) override;
+	virtual bool UpdateStates() const override;
+
+protected:
+	/** List of files to remove from shelved files in changelist state */
+	TArray<FString> FilesToRemove;
+
+	/** Changelist to be updated */
+	FPerforceSourceControlChangelist ChangelistToUpdate;
+};
+
+class FPerforceUnshelveWorker : public IPerforceSourceControlWorker
+{
+public:
+	virtual ~FPerforceUnshelveWorker() {}
+	// IPerforceSourceControlWorker interface
+	virtual FName GetName() const override;
+	virtual bool Execute(class FPerforceSourceControlCommand& InCommand) override;
+	virtual bool UpdateStates() const override;
+
+protected:
+	/** Changelist to be updated */
+	FPerforceSourceControlChangelist ChangelistToUpdate;
+
+	/** List of files states after update */
+	TArray<FPerforceSourceControlState> ChangelistFilesStates;
+};
diff --git a/Engine/Plugins/Developer/PerforceSourceControl/Source/PerforceSourceControl/Private/PerforceSourceControlState.cpp b/Engine/Plugins/Developer/PerforceSourceControl/Source/PerforceSourceControl/Private/PerforceSourceControlState.cpp
--- a/Engine/Plugins/Developer/PerforceSourceControl/Source/PerforceSourceControl/Private/PerforceSourceControlState.cpp
+++ b/Engine/Plugins/Developer/PerforceSourceControl/Source/PerforceSourceControl/Private/PerforceSourceControlState.cpp
@@ -336,4 +336,92 @@ bool FPerforceSourceControlState::CanRevert() const
 	return CanCheckIn();
 }
 
+void FPerforceSourceControlState::Update(const FPerforceSourceControlState& InOther, const FDateTime* InTimeStamp /* = nullptr */)
+{
+	check(InOther.LocalFilename == LocalFilename);
+
+	if (InOther.History.Num() != 0)
+	{
+		History = InOther.History;
+	}
+
+	if (InOther.DepotFilename.Len() != 0)
+	{
+		DepotFilename = InOther.DepotFilename;
+	}
+
+	if (InOther.OtherUserCheckedOut.Len() != 0)
+	{
+		OtherUserCheckedOut = InOther.OtherUserCheckedOut;
+	}
+
+	if (InOther.State != EPerforceState::DontCare)
+	{
+		State = InOther.State;
+	}
+
+	if (InOther.DepotRevNumber != INVALID_REVISION)
+	{
+		DepotRevNumber = InOther.DepotRevNumber;
+	}
+
+	if (InOther.LocalRevNumber != INVALID_REVISION)
+	{
+		LocalRevNumber = InOther.LocalRevNumber;
+	}
+
+	if (InOther.PendingResolveRevNumber != INVALID_REVISION)
+	{
+		PendingResolveRevNumber = InOther.PendingResolveRevNumber;
+	}
+
+	// This flag is true when the current object has been modified.
+	// This is not necessarily known at all times, so we will play it safe and use an OR here.
+	bModifed |= InOther.bModifed;
+
+	// Here we will assume that the type info is always properly intiialized in InOther
+	bBinary = InOther.bBinary;
+	bExclusiveCheckout = InOther.bExclusiveCheckout;
+
+	if (InOther.Changelist.IsInitialized())
+	{
+		Changelist = InOther.Changelist;
+	}
+
+	if (InTimeStamp)
+	{
+		TimeStamp = *InTimeStamp;
+	}
+
+	if (InOther.HeadBranch.Len() != 0)
+	{
+		HeadBranch = InOther.HeadBranch;
+	}
+
+	if (InOther.HeadAction.Len() != 0)
+	{
+		HeadAction = InOther.HeadAction;
+	}
+
+	if (InOther.HeadModTime != 0)
+	{
+		HeadModTime = InOther.HeadModTime;
+	}
+
+	if (InOther.HeadChangeList != 0)
+	{
+		HeadChangeList = InOther.HeadChangeList;
+	}
+
+	if (InOther.CheckedOutBranches.Num() != 0)
+	{
+		CheckedOutBranches = InOther.CheckedOutBranches;
+	}
+
+	if (InOther.OtherUserBranchCheckedOuts.Len() != 0)
+	{
+		OtherUserBranchCheckedOuts = InOther.OtherUserBranchCheckedOuts;
+	}
+}
+
 #undef LOCTEXT_NAMESPACE
diff --git a/Engine/Plugins/Developer/PerforceSourceControl/Source/PerforceSourceControl/Private/PerforceSourceControlState.h b/Engine/Plugins/Developer/PerforceSourceControl/Source/PerforceSourceControl/Private/PerforceSourceControlState.h
--- a/Engine/Plugins/Developer/PerforceSourceControl/Source/PerforceSourceControl/Private/PerforceSourceControlState.h
+++ b/Engine/Plugins/Developer/PerforceSourceControl/Source/PerforceSourceControl/Private/PerforceSourceControlState.h
@@ -6,6 +6,8 @@
 #include "ISourceControlState.h"
 #include "ISourceControlRevision.h"
 
+#include "PerforceSourceControlChangelist.h"
+
 class FPerforceSourceControlRevision;
 
 namespace EPerforceState
@@ -54,6 +56,8 @@ public:
 		, LocalRevNumber(INVALID_REVISION)
 		, PendingResolveRevNumber(INVALID_REVISION)
 		, bModifed(false)
+		, bBinary(false)
+		, bExclusiveCheckout(false)
 		, TimeStamp(0)
 		, HeadModTime(0)
 		, HeadChangeList(0)
@@ -106,6 +110,9 @@ public:
 		State = InState;
 	}
 
+	/** Fills in missing information from other state */
+	void Update(const FPerforceSourceControlState& InOther, const FDateTime* TimeStamp = nullptr);
+
 public:
 	/** History of the item, if any */
 	TArray< TSharedRef<FPerforceSourceControlRevision, ESPMode::ThreadSafe> > History;
@@ -131,6 +138,9 @@ public:
 	/** Pending rev number with which a file must be resolved, INVALID_REVISION if no resolve pending */
 	int PendingResolveRevNumber;
 
+	/** Changelist containing this file */
+	FPerforceSourceControlChangelist Changelist;
+
 	/** Modified from depot version */
 	bool bModifed;
 
diff --git a/Engine/Source/Developer/SourceControl/Public/SourceControlOperations.h b/Engine/Source/Developer/SourceControl/Public/SourceControlOperations.h
--- a/Engine/Source/Developer/SourceControl/Public/SourceControlOperations.h
+++ b/Engine/Source/Developer/SourceControl/Public/SourceControlOperations.h
@@ -535,5 +535,58 @@ public:
 	}
 };
 
+/**
+ * Operation used to shelve files in a changelist
+ */
+class FShelve : public FSourceControlOperationBase
+{
+public:
+	// ISourceControlOperation interface
+	virtual FName GetName() const override
+	{
+		return "Shelve";
+	}
+
+	virtual FText GetInProgressString() const override
+	{
+		return LOCTEXT("SourceControl_ShelveOperation", "Shelving files in changelist...");
+	}
+};
+
+/**
+ * Operation used to unshelve files from a changelist
+ */
+class FUnshelve : public FSourceControlOperationBase
+{
+public:
+	// ISourceControlOperation interface
+	virtual FName GetName() const override 
+	{
+		return "Unshelve";
+	}
+
+	virtual FText GetInProgressString() const override
+	{
+		return LOCTEXT("SourceControl_UnshelveOperation", "Unshelving files from changelist...");
+	}
+};
+
+/**
+ * Operation used to delete shelved files from a changelist
+ */
+class FDeleteShelved : public FSourceControlOperationBase
+{
+public:
+	// ISourceControlOperation interface
+	virtual FName GetName() const override
+	{
+		return "DeleteShelved";
+	}
+
+	virtual FText GetInProgressString() const override
+	{
+		return LOCTEXT("SourceControl_DeleteShelvedOperation", "Deleting shelved files from changelist...");
+	}
+};
 
 #undef LOCTEXT_NAMESPACE
diff --git a/Engine/Source/Editor/SourceControlWindows/Private/SSourceControlChangelists.cpp b/Engine/Source/Editor/SourceControlWindows/Private/SSourceControlChangelists.cpp
--- a/Engine/Source/Editor/SourceControlWindows/Private/SSourceControlChangelists.cpp
+++ b/Engine/Source/Editor/SourceControlWindows/Private/SSourceControlChangelists.cpp
@@ -28,7 +28,7 @@
 #define LOCTEXT_NAMESPACE "SourceControlChangelist"
 
 //////////////////////////////
-static TSharedRef<SWidget> GetSCCFileWidget(FSourceControlStateRef InFileState)
+static TSharedRef<SWidget> GetSCCFileWidget(FSourceControlStateRef InFileState, bool bIsShelvedFile = false)
 {
 	const FSlateBrush* IconBrush = FEditorStyle::GetBrush("ContentBrowser.ColumnViewAssetIcon");
 
@@ -37,16 +37,17 @@ static TSharedRef<SWidget> GetSCCFileWidget(FSourceControlStateRef InFileState)
 	const float IconOverlaySize = IconBrush->ImageSize.X * ICON_SCALING_FACTOR;
 
 	return SNew(SOverlay)
-
 		// The actual icon
 		+ SOverlay::Slot()
 		[
 			SNew(SImage)
 			.Image(IconBrush)
+			.ColorAndOpacity_Lambda([bIsShelvedFile]() -> FSlateColor {
+					return FSlateColor(bIsShelvedFile ? FColor::Yellow : FColor::White);
+				})
 		]
-
-	// Source control state
-	+ SOverlay::Slot()
+		// Source control state
+		+ SOverlay::Slot()
 		.HAlign(HAlign_Left)
 		.VAlign(VAlign_Top)
 		[
@@ -96,6 +97,19 @@ struct FChangelistTreeItem : public IChangelistTreeItem
 	FSourceControlChangelistStateRef ChangelistState;
 };
 
+struct FShelvedChangelistTreeItem : public IChangelistTreeItem
+{
+	FShelvedChangelistTreeItem()
+	{
+		Type = IChangelistTreeItem::ShelvedChangelist;
+	}
+
+	FText GetDisplayText() const
+	{
+		return LOCTEXT("SourceControl_ShelvedFiles", "Shelved Items");
+	}
+};
+
 bool GetAssetData(const FString& InPackageName, const FString& InFileName, TArray<FAssetData>& OutAssets)
 {
 	OutAssets.Reset();
@@ -246,6 +260,22 @@ private:
 	FColor DisplayColor;
 };
 
+struct FShelvedFileTreeItem : public IChangelistTreeItem
+{
+	FShelvedFileTreeItem(FSourceControlStateRef InFileState)
+		: FileState(InFileState)
+	{
+		Type = IChangelistTreeItem::ShelvedFile;
+	}
+
+	FText GetDisplayName() const
+	{
+		return FText::FromString(FileState->GetFilename());
+	}
+
+	FSourceControlStateRef FileState;
+};
+
 SSourceControlChangelistsWidget::SSourceControlChangelistsWidget()
 {
 }
@@ -295,6 +325,7 @@ void SSourceControlChangelistsWidget::RequestRefresh()
 		TSharedRef<FUpdatePendingChangelistsStatus, ESPMode::ThreadSafe> UpdatePendingChangelistsOperation = ISourceControlOperation::Create<FUpdatePendingChangelistsStatus>();
 		UpdatePendingChangelistsOperation->SetUpdateAllChangelists(true);
 		UpdatePendingChangelistsOperation->SetUpdateFilesStates(true);
+		UpdatePendingChangelistsOperation->SetUpdateShelvedFilesStates(true);
 
 		ISourceControlProvider& SourceControlProvider = ISourceControlModule::Get().GetProvider();
 		SourceControlProvider.Execute(UpdatePendingChangelistsOperation, EConcurrency::Asynchronous, FSourceControlOperationComplete::CreateSP(this, &SSourceControlChangelistsWidget::OnChangelistsStatusUpdated));
@@ -337,6 +368,18 @@ void SSourceControlChangelistsWidget::Refresh()
 				ChangelistTreeItem->AddChild(FileTreeItem);
 			}
 
+			if (ChangelistState->GetShelvedFilesStates().Num() > 0)
+			{
+				FChangelistTreeItemRef ShelvedChangelistTreeItem = MakeShareable(new FShelvedChangelistTreeItem());
+				ChangelistTreeItem->AddChild(ShelvedChangelistTreeItem);
+
+				for (FSourceControlStateRef ShelvedFileRef : ChangelistState->GetShelvedFilesStates())
+				{
+					FChangelistTreeItemRef ShelvedFileTreeItem = MakeShareable(new FShelvedFileTreeItem(ShelvedFileRef));
+					ShelvedChangelistTreeItem->AddChild(ShelvedFileTreeItem);
+				}
+			}
+
 			ChangelistsNodes.Add(ChangelistTreeItem);
 		}
 
@@ -416,7 +459,7 @@ FSourceControlChangelistPtr SSourceControlChangelistsWidget::GetCurrentChangelis
 	return ChangelistState ? (FSourceControlChangelistPtr)(ChangelistState->GetChangelist()) : nullptr;
 }
 
-FSourceControlChangelistStatePtr SSourceControlChangelistsWidget::GetChangelistStateFromFileSelection()
+FSourceControlChangelistStatePtr SSourceControlChangelistsWidget::GetChangelistStateFromSelection()
 {
 	if (!TreeView)
 	{
@@ -425,19 +468,27 @@ FSourceControlChangelistStatePtr SSourceControlChangelistsWidget::GetChangelistS
 
 	TArray<FChangelistTreeItemPtr> SelectedItems = TreeView->GetSelectedItems();
 
-	if (SelectedItems.Num() == 0 || SelectedItems[0]->GetTreeItemType() != IChangelistTreeItem::File)
+	if (SelectedItems.Num() == 0 || SelectedItems[0]->GetTreeItemType() == IChangelistTreeItem::Invalid)
 	{
 		return nullptr;
 	}
-	else
+
+	FChangelistTreeItemPtr Item = SelectedItems[0];
+
+	while (Item && Item->GetTreeItemType() != IChangelistTreeItem::Invalid)
 	{
-		return StaticCastSharedPtr<FChangelistTreeItem>(SelectedItems[0]->GetParent())->ChangelistState;
+		if (Item->GetTreeItemType() == IChangelistTreeItem::Changelist)
+			return StaticCastSharedPtr<FChangelistTreeItem>(Item)->ChangelistState;
+		else
+			Item = Item->GetParent();
 	}
+
+	return nullptr;
 }
 
-FSourceControlChangelistPtr SSourceControlChangelistsWidget::GetChangelistFromFileSelection()
+FSourceControlChangelistPtr SSourceControlChangelistsWidget::GetChangelistFromSelection()
 {
-	FSourceControlChangelistStatePtr ChangelistState = GetChangelistStateFromFileSelection();
+	FSourceControlChangelistStatePtr ChangelistState = GetChangelistStateFromSelection();
 	return ChangelistState ? (FSourceControlChangelistPtr)(ChangelistState->GetChangelist()) : nullptr;
 }
 
@@ -461,6 +512,34 @@ TArray<FString> SSourceControlChangelistsWidget::GetSelectedFiles()
 	}
 }
 
+TArray<FString> SSourceControlChangelistsWidget::GetSelectedShelvedFiles()
+{
+	TArray<FString> ShelvedFiles;
+	TArray<FChangelistTreeItemPtr> SelectedItems = TreeView->GetSelectedItems();
+	
+	if (SelectedItems.Num() > 0)
+	{
+		if (SelectedItems[0]->GetTreeItemType() == IChangelistTreeItem::ShelvedChangelist)
+		{
+			check(SelectedItems.Num() == 1);
+			const TArray<FChangelistTreeItemPtr>& ShelvedChildren = SelectedItems[0]->GetChildren();
+			for (FChangelistTreeItemPtr Item : ShelvedChildren)
+			{
+				ShelvedFiles.Add(StaticCastSharedPtr<FShelvedFileTreeItem>(Item)->FileState->GetFilename());
+			}
+		}
+		else if (SelectedItems[0]->GetTreeItemType() == IChangelistTreeItem::ShelvedFile)
+		{
+			for (FChangelistTreeItemPtr Item : SelectedItems)
+			{
+				ShelvedFiles.Add(StaticCastSharedPtr<FShelvedFileTreeItem>(Item)->FileState->GetFilename());
+			}
+		}
+	}
+
+	return ShelvedFiles;
+}
+
 void SSourceControlChangelistsWidget::OnNewChangelist()
 {
 	FText ChangelistDescription;
@@ -480,7 +559,6 @@ void SSourceControlChangelistsWidget::OnNewChangelist()
 	NewChangelistOperation->SetDescription(ChangelistDescription);
 
 	SourceControlProvider.Execute(NewChangelistOperation);
-	Refresh();
 }
 
 void SSourceControlChangelistsWidget::OnDeleteChangelist()
@@ -492,7 +570,6 @@ void SSourceControlChangelistsWidget::OnDeleteChangelist()
 
 	ISourceControlProvider& SourceControlProvider = ISourceControlModule::Get().GetProvider();
 	SourceControlProvider.Execute(ISourceControlOperation::Create<FDeleteChangelist>(), GetCurrentChangelist());
-	Refresh();
 }
 
 bool SSourceControlChangelistsWidget::CanDeleteChangelist()
@@ -528,8 +605,6 @@ void SSourceControlChangelistsWidget::OnEditChangelist()
 
 	ISourceControlProvider& SourceControlProvider = ISourceControlModule::Get().GetProvider();
 	SourceControlProvider.Execute(EditChangelistOperation, ChangelistState->GetChangelist());
-
-	Refresh();
 }
 
 void SSourceControlChangelistsWidget::OnRevertUnchanged()
@@ -537,19 +612,7 @@ void SSourceControlChangelistsWidget::OnRevertUnchanged()
 	ISourceControlProvider& SourceControlProvider = ISourceControlModule::Get().GetProvider();
 
 	auto RevertUnchangedOperation = ISourceControlOperation::Create<FRevertUnchanged>();
-
-	if(GetCurrentChangelist() != nullptr)
-	{
-		// Operation on full changelist
-		SourceControlProvider.Execute(RevertUnchangedOperation, GetCurrentChangelist());
-	}
-	else
-	{
-		// Operation on files in a changelist
-		SourceControlProvider.Execute(RevertUnchangedOperation, GetChangelistFromFileSelection(), GetSelectedFiles());
-	}
-
-	Refresh();
+	SourceControlProvider.Execute(RevertUnchangedOperation, GetChangelistFromSelection(), GetSelectedFiles());
 }
 
 void SSourceControlChangelistsWidget::OnRevert()
@@ -579,19 +642,28 @@ void SSourceControlChangelistsWidget::OnRevert()
 
 	ISourceControlProvider& SourceControlProvider = ISourceControlModule::Get().GetProvider();
 	auto RevertOperation = ISourceControlOperation::Create<FRevert>();
+	SourceControlProvider.Execute(RevertOperation, GetChangelistFromSelection(), GetSelectedFiles());
+}
 
-	if (GetCurrentChangelist() != nullptr)
-	{
-		// Operation on full changelist
-		SourceControlProvider.Execute(RevertOperation, GetCurrentChangelist());
-	}
-	else
-	{
-		// Operation on files in a changelist
-		SourceControlProvider.Execute(RevertOperation, GetChangelistFromFileSelection(), GetSelectedFiles());
-	}
+void SSourceControlChangelistsWidget::OnShelve()
+{
+	ISourceControlProvider& SourceControlProvider = ISourceControlModule::Get().GetProvider();
+	auto ShelveOperation = ISourceControlOperation::Create<FShelve>();
+	SourceControlProvider.Execute(ShelveOperation, GetChangelistFromSelection(), GetSelectedFiles());
+}
 
-	Refresh();
+void SSourceControlChangelistsWidget::OnUnshelve()
+{
+	ISourceControlProvider& SourceControlProvider = ISourceControlModule::Get().GetProvider();
+	auto UnshelveOperation = ISourceControlOperation::Create<FUnshelve>();
+	SourceControlProvider.Execute(UnshelveOperation, GetChangelistFromSelection(), GetSelectedShelvedFiles());
+}
+
+void SSourceControlChangelistsWidget::OnDeleteShelvedFiles()
+{
+	ISourceControlProvider& SourceControlProvider = ISourceControlModule::Get().GetProvider();
+	auto DeleteShelvedOperation = ISourceControlOperation::Create<FDeleteShelved>();
+	SourceControlProvider.Execute(DeleteShelvedOperation, GetChangelistFromSelection(), GetSelectedShelvedFiles());
 }
 
 void SSourceControlChangelistsWidget::OnSubmitChangelist()
@@ -672,6 +744,16 @@ bool SSourceControlChangelistsWidget::CanDiffAgainstDepot()
 	return GetSelectedFiles().Num() == 1;
 }
 
+void SSourceControlChangelistsWidget::OnDiffAgainstWorkspace()
+{
+
+}
+
+bool SSourceControlChangelistsWidget::CanDiffAgainstWorkspace()
+{
+	return GetSelectedShelvedFiles().Num() == 1;
+}
+
 TSharedPtr<SWidget> SSourceControlChangelistsWidget::OnOpenContextMenu()
 {
 	UToolMenus* ToolMenus = UToolMenus::Get();
@@ -686,7 +768,9 @@ TSharedPtr<SWidget> SSourceControlChangelistsWidget::OnOpenContextMenu()
 	UToolMenu* Menu = ToolMenus->GenerateMenu(MenuName, Context);
 		
 	bool bHasSelectedChangelist = (GetCurrentChangelist() != nullptr);
-	bool bHasEmptySelection = (GetCurrentChangelist() == nullptr && GetSelectedFiles().Num() == 0);
+	bool bHasSelectedFiles = (GetSelectedFiles().Num() > 0);
+	bool bHasSelectedShelvedFiles = (GetSelectedShelvedFiles().Num() > 0);
+	bool bHasEmptySelection = (!bHasSelectedChangelist && !bHasSelectedFiles && !bHasSelectedShelvedFiles);
 
 	FToolMenuSection& Section = Menu->AddSection("Source Control");
 	
@@ -700,7 +784,7 @@ TSharedPtr<SWidget> SSourceControlChangelistsWidget::OnOpenContextMenu()
 	}
 
 	// This can appear on both files & changelist
-	if (!bHasEmptySelection)
+	if (bHasSelectedChangelist || bHasSelectedFiles)
 	{
 		Section.AddMenuEntry("RevertUnchanged", LOCTEXT("SourceControl_RevertUnchanged", "Revert Unchanged"), LOCTEXT("SourceControl_Revert_Unchanged_Tooltip", "Reverts unchanged files & changelists"), FSlateIcon(),
 			FUIAction(FExecuteAction::CreateSP(this, &SSourceControlChangelistsWidget::OnRevertUnchanged)));
@@ -709,6 +793,36 @@ TSharedPtr<SWidget> SSourceControlChangelistsWidget::OnOpenContextMenu()
 			FUIAction(FExecuteAction::CreateSP(this, &SSourceControlChangelistsWidget::OnRevert)));
 	}
 
+	if(bHasSelectedFiles || bHasSelectedShelvedFiles || (bHasSelectedChangelist && (GetCurrentChangelistState()->GetFilesStates().Num() > 0 || GetCurrentChangelistState()->GetShelvedFilesStates().Num() > 0)))
+	{
+		Section.AddSeparator("Shelve");
+	}
+
+	if(bHasSelectedFiles || (bHasSelectedChangelist && GetCurrentChangelistState()->GetFilesStates().Num() > 0))
+	{
+		Section.AddMenuEntry("Shelve", LOCTEXT("SourceControl_Shelve", "Shelve Files"), LOCTEXT("SourceControl_Shelve_Tooltip", "Shelves the changelist or the selected files"), FSlateIcon(),
+			FUIAction(FExecuteAction::CreateSP(this, &SSourceControlChangelistsWidget::OnShelve)));
+	}
+
+	if (bHasSelectedShelvedFiles || (bHasSelectedChangelist && GetCurrentChangelistState()->GetShelvedFilesStates().Num() > 0))
+	{
+		Section.AddMenuEntry("Unshelve", LOCTEXT("SourceControl_Unshelve", "Unshelve Files"), LOCTEXT("SourceControl_Unshelve_Tooltip", "Unshelve selected files or changelist"), FSlateIcon(),
+			FUIAction(FExecuteAction::CreateSP(this, &SSourceControlChangelistsWidget::OnUnshelve)));
+
+		Section.AddMenuEntry("DeleteShelved", LOCTEXT("SourceControl_DeleteShelved", "Delete Shelved Files"), LOCTEXT("SourceControl_DeleteShelved_Tooltip", "Delete selected shelved files or all from changelist"), FSlateIcon(),
+			FUIAction(FExecuteAction::CreateSP(this, &SSourceControlChangelistsWidget::OnDeleteShelvedFiles)));
+	}
+
+	// Shelved files-only operations
+	if (bHasSelectedShelvedFiles)
+	{
+		// Diff against workspace
+		Section.AddMenuEntry("DiffAgainstWorkspace", LOCTEXT("SourceControl_DiffAgainstWorkspace", "Diff Against Workspace Files"), LOCTEXT("SourceControl_DiffAgainstWorkspace_Tooltip", "Diff shelved file against the (local) workspace file"), FSlateIcon(),
+			FUIAction(
+				FExecuteAction::CreateSP(this, &SSourceControlChangelistsWidget::OnDiffAgainstWorkspace),
+				FCanExecuteAction::CreateSP(this, &SSourceControlChangelistsWidget::CanDiffAgainstWorkspace)));
+	}
+
 	if (bHasEmptySelection || bHasSelectedChangelist)
 	{
 		Section.AddSeparator("Changelists");
@@ -733,7 +847,7 @@ TSharedPtr<SWidget> SSourceControlChangelistsWidget::OnOpenContextMenu()
 	}
 
 	// Files-only operations
-	if (!bHasEmptySelection && !bHasSelectedChangelist)
+	if(bHasSelectedFiles)
 	{
 		Section.AddSeparator("Files");
 
@@ -998,6 +1112,144 @@ private:
 	FFileTreeItem* TreeItem;
 };
 
+class SShelvedChangelistTableRow : public SMultiColumnTableRow<FChangelistTreeItemPtr>
+{
+public:
+	SLATE_BEGIN_ARGS(SShelvedChangelistTableRow)
+		: _TreeItemToVisualize()
+	{}
+	SLATE_ARGUMENT(FChangelistTreeItemPtr, TreeItemToVisualize)
+		SLATE_END_ARGS()
+
+public:
+	/**
+	* Construct child widgets that comprise this widget.
+	*
+	* @param InArgs Declaration from which to construct this widget.
+	*/
+	void Construct(const FArguments& InArgs, const TSharedRef<STableViewBase>& InOwner)
+	{
+		TreeItem = static_cast<FShelvedChangelistTreeItem*>(InArgs._TreeItemToVisualize.Get());
+
+		auto Args = FSuperRowType::FArguments();
+		SMultiColumnTableRow<FChangelistTreeItemPtr>::Construct(Args, InOwner);
+	}
+
+	// SMultiColumnTableRow overrides
+	virtual TSharedRef<SWidget> GenerateWidgetForColumn(const FName& ColumnName) override
+	{
+		if (ColumnName == TEXT("Change"))
+		{
+			return SNew(SHorizontalBox)
+				+ SHorizontalBox::Slot()
+					.AutoWidth()
+					.VAlign(VAlign_Center)
+					.Padding(5, 0, 4, 0)
+					[
+						SNew(SExpanderArrow, SharedThis(this))
+					]
+				+ SHorizontalBox::Slot()
+					.AutoWidth()
+					.VAlign(VAlign_Center)
+					.Padding(5, 0, 0, 0)
+					[
+						SNew(SImage)
+						.Image(FEditorStyle::GetBrush("SourceControl.Changelist"))
+					]
+				+ SHorizontalBox::Slot()
+					.Padding(2.0f, 0.0f)
+					.VAlign(VAlign_Center)
+					[
+						SNew(STextBlock)
+						.Text(this, &SShelvedChangelistTableRow::GetText)
+					];
+		}
+		else
+		{
+			return SNullWidget::NullWidget;
+		}
+	}
+
+protected:
+	FText GetText() const
+	{
+		return TreeItem->GetDisplayText();
+	}
+
+private:
+	/** The info about the widget that we are visualizing. */
+	FShelvedChangelistTreeItem* TreeItem;
+};
+
+class SShelvedFileTableRow : public SMultiColumnTableRow<FChangelistTreeItemPtr>
+{
+public:
+	SLATE_BEGIN_ARGS(SShelvedFileTableRow)
+		: _TreeItemToVisualize()
+	{}
+	SLATE_ARGUMENT(FChangelistTreeItemPtr, TreeItemToVisualize)
+		SLATE_END_ARGS()
+
+public:
+	/**
+	* Construct child widgets that comprise this widget.
+	*
+	* @param InArgs Declaration from which to construct this widget.
+	*/
+	void Construct(const FArguments& InArgs, const TSharedRef<STableViewBase>& InOwner)
+	{
+		TreeItem = static_cast<FShelvedFileTreeItem*>(InArgs._TreeItemToVisualize.Get());
+
+		auto Args = FSuperRowType::FArguments();
+		FSuperRowType::Construct(Args, InOwner);
+	}
+
+	// SMultiColumnTableRow overrides
+	virtual TSharedRef<SWidget> GenerateWidgetForColumn(const FName& ColumnName) override
+	{
+		const bool bIsShelvedFile = true;
+
+		if (ColumnName == TEXT("Change")) // eq. to name
+		{
+			return SNew(SHorizontalBox)
+
+				// Icon
+				+ SHorizontalBox::Slot()
+				.AutoWidth()
+				.Padding(60, 0, 4, 0)
+				[
+					GetSCCFileWidget(TreeItem->FileState, bIsShelvedFile)
+				]
+
+			+ SHorizontalBox::Slot()
+				.AutoWidth()
+				.VAlign(VAlign_Center)
+				[
+					SNew(STextBlock)
+					.Text(LOCTEXT("SourceControl_DefaultNameForShelvedFiles", "Unavailable"))
+				];
+		}
+		else if (ColumnName == TEXT("Description"))
+		{
+			return SNew(STextBlock)
+				.Text(this, &SShelvedFileTableRow::GetDisplayName);
+		}
+		else
+		{
+			return SNullWidget::NullWidget;
+		}
+	}
+
+	FText GetDisplayName() const
+	{
+		return TreeItem->GetDisplayName();
+	}
+
+private:
+	/** The info about the widget that we are visualizing. */
+	FShelvedFileTreeItem* TreeItem;
+};
+
 TSharedRef<ITableRow> SSourceControlChangelistsWidget::OnGenerateRow(FChangelistTreeItemPtr InTreeItem, const TSharedRef<STableViewBase>& OwnerTable)
 {
 	switch (InTreeItem->GetTreeItemType())
@@ -1011,6 +1263,14 @@ TSharedRef<ITableRow> SSourceControlChangelistsWidget::OnGenerateRow(FChangelist
 			.TreeItemToVisualize(InTreeItem)
 			.OnDragDetected(this, &SSourceControlChangelistsWidget::OnFilesDragged);
 
+	case IChangelistTreeItem::ShelvedChangelist:
+		return SNew(SShelvedChangelistTableRow, OwnerTable)
+			.TreeItemToVisualize(InTreeItem);
+
+	case IChangelistTreeItem::ShelvedFile:
+		return SNew(SShelvedFileTableRow, OwnerTable)
+			.TreeItemToVisualize(InTreeItem);
+
 	default:
 		check(false);
 	};
@@ -1044,5 +1304,4 @@ void SSourceControlChangelistsWidget::OnGetChildren(FChangelistTreeItemPtr InPar
 }
 
 
-
 #undef LOCTEXT_NAMESPACE
\ No newline at end of file
diff --git a/Engine/Source/Editor/SourceControlWindows/Private/SSourceControlChangelists.h b/Engine/Source/Editor/SourceControlWindows/Private/SSourceControlChangelists.h
--- a/Engine/Source/Editor/SourceControlWindows/Private/SSourceControlChangelists.h
+++ b/Engine/Source/Editor/SourceControlWindows/Private/SSourceControlChangelists.h
@@ -19,7 +19,9 @@ struct IChangelistTreeItem : TSharedFromThis<IChangelistTreeItem>
 	{
 		Invalid,
 		Changelist,
-		File
+		File,
+		ShelvedChangelist, // container for shelved files
+		ShelvedFile
 	};
 
 	/** Get this item's parent. Can be nullptr. */
@@ -103,12 +105,15 @@ private:
 	/** Returns the currently selected changelist state ptr or null in invalid cases */
 	FSourceControlChangelistStatePtr GetCurrentChangelistState();
 	FSourceControlChangelistPtr GetCurrentChangelist();
-	FSourceControlChangelistStatePtr GetChangelistStateFromFileSelection();
-	FSourceControlChangelistPtr GetChangelistFromFileSelection();
+	FSourceControlChangelistStatePtr GetChangelistStateFromSelection();
+	FSourceControlChangelistPtr GetChangelistFromSelection();
 
 	/** Returns list of currently selected files */
 	TArray<FString> GetSelectedFiles();
 
+	/** Returns list of currently selected shelved files */
+	TArray<FString> GetSelectedShelvedFiles();
+
 	/** Changelist operations */
 	void OnNewChangelist();
 	void OnDeleteChangelist();
@@ -120,6 +125,11 @@ private:
 	/** Changelist & File operations */
 	void OnRevertUnchanged();
 	void OnRevert();
+	void OnShelve();
+
+	/** Changelist & shelved files operations */
+	void OnUnshelve();
+	void OnDeleteShelvedFiles();
 
 	/** Files operations */
 	void OnLocateFile();
@@ -128,6 +138,10 @@ private:
 	void OnDiffAgainstDepot();
 	bool CanDiffAgainstDepot();
 
+	/** Shelved files operations */
+	void OnDiffAgainstWorkspace();
+	bool CanDiffAgainstWorkspace();
+
 	void OnSourceControlStateChanged();
 	void OnSourceControlProviderChanged(ISourceControlProvider& OldProvider, ISourceControlProvider& NewProvider);
 	void OnChangelistsStatusUpdated(const FSourceControlOperationRef& InOperation, ECommandResult::Type InType);
