Improved behavior of changelist operations on default changelist - edit & submit now work as intended.
Fixed an issue when trying to delete the default changelist which triggered a disconnection from the source control provider.
Improved some UI in the changelist context menu to clarify what is available or not.

diff --git a/Engine/Plugins/Developer/PerforceSourceControl/Source/PerforceSourceControl/Private/PerforceSourceControlOperations.cpp b/Engine/Plugins/Developer/PerforceSourceControl/Source/PerforceSourceControl/Private/PerforceSourceControlOperations.cpp
--- a/Engine/Plugins/Developer/PerforceSourceControl/Source/PerforceSourceControl/Private/PerforceSourceControlOperations.cpp
+++ b/Engine/Plugins/Developer/PerforceSourceControl/Source/PerforceSourceControl/Private/PerforceSourceControlOperations.cpp
@@ -17,6 +17,7 @@
 #include "Algo/AnyOf.h"
 #include "Algo/IndexOf.h"
 #include "Algo/Find.h"
+#include "Algo/Transform.h"
 
 #define LOCTEXT_NAMESPACE "PerforceSourceControl"
 
@@ -388,6 +389,23 @@ static bool RunReopenCommand(FPerforceSourceControlCommand& InCommand, const TAr
 	return bCommandSuccessful;
 }
 
+static bool RemoveFilesFromChangelist(const TMap<FString, EPerforceState::Type>& Results, TSharedRef<FPerforceSourceControlChangelistState, ESPMode::ThreadSafe>& ChangelistState)
+{
+	return ChangelistState->Files.RemoveAll([&Results](FSourceControlStateRef& State) -> bool
+		{
+			return Algo::AnyOf(Results, [&State](auto& Result) {
+				return State->GetFilename() == Result.Key;
+				});
+		}) > 0;
+}
+
+static bool RemoveFilesFromChangelist(const TMap<FString, EPerforceState::Type>& Results, const FPerforceSourceControlChangelist& Changelist)
+{
+	FPerforceSourceControlModule& PerforceSourceControl = FPerforceSourceControlModule::Get();
+	TSharedRef<FPerforceSourceControlChangelistState, ESPMode::ThreadSafe> ChangelistState = PerforceSourceControl.GetProvider().GetStateInternal(Changelist);
+	return RemoveFilesFromChangelist(Results, ChangelistState);
+}
+
 bool FPerforceCheckInWorker::Execute(FPerforceSourceControlCommand& InCommand)
 {
 	FScopedPerforceConnection ScopedConnection(InCommand);
@@ -398,6 +416,8 @@ bool FPerforceCheckInWorker::Execute(FPerforceSourceControlCommand& InCommand)
 		check(InCommand.Operation->GetName() == GetName());
 		TSharedRef<FCheckIn, ESPMode::ThreadSafe> Operation = StaticCastSharedRef<FCheckIn>(InCommand.Operation);
 
+		TArray<FString> FilesToSubmit = InCommand.Files;
+
 		FPerforceSourceControlChangelist ChangeList(InCommand.Changelist);
 		TArray<FString> ReopenedFiles;
 
@@ -405,11 +425,21 @@ bool FPerforceCheckInWorker::Execute(FPerforceSourceControlCommand& InCommand)
 
 		if (InCommand.Changelist.IsDefault())
 		{
+			// If the command has specified the default changelist but no files, then get all files from the default changelist
+			if (FilesToSubmit.Num() == 0 && InCommand.Changelist.IsInitialized())
+			{
+				FPerforceSourceControlModule& PerforceSourceControl = FPerforceSourceControlModule::Get();
+				TSharedRef<FPerforceSourceControlChangelistState, ESPMode::ThreadSafe> DefaultChangelistState = PerforceSourceControl.GetProvider().GetStateInternal(InCommand.Changelist);
+				Algo::Transform(DefaultChangelistState->Files, FilesToSubmit, [](const auto& FileState) {
+					return FileState->GetFilename();
+					});
+			}
+
 			int32 NewChangeList = Connection.CreatePendingChangelist(Operation->GetDescription(), TArray<FString>(), FOnIsCancelled::CreateRaw(&InCommand, &FPerforceSourceControlCommand::IsCanceled), InCommand.ResultInfo.ErrorMessages);
 			if (NewChangeList > 0)
 			{
 				ChangeList = FPerforceSourceControlChangelist(NewChangeList);
-				InCommand.bCommandSuccessful = RunReopenCommand(InCommand, InCommand.Files, ChangeList, &ReopenedFiles);
+				InCommand.bCommandSuccessful = RunReopenCommand(InCommand, FilesToSubmit, ChangeList, &ReopenedFiles);
 			}
 			else
 			{
@@ -440,7 +470,7 @@ bool FPerforceCheckInWorker::Execute(FPerforceSourceControlCommand& InCommand)
 				FPerforceSourceControlProvider& Provider = PerforceSourceControl.GetProvider();
 
 				TArray<TSharedRef<ISourceControlState, ESPMode::ThreadSafe>> States;
-				Provider.GetState(InCommand.Files, States, EStateCacheUsage::Use);
+				Provider.GetState(FilesToSubmit, States, EStateCacheUsage::Use);
 				for (const auto& State : States)
 				{
 					if (State->IsDeleted())
@@ -451,11 +481,12 @@ bool FPerforceCheckInWorker::Execute(FPerforceSourceControlCommand& InCommand)
 
 				StaticCastSharedRef<FCheckIn>(InCommand.Operation)->SetSuccessMessage(ParseSubmitResults(Records));
 
-				for(auto Iter(InCommand.Files.CreateIterator()); Iter; Iter++)
+				for(auto Iter(FilesToSubmit.CreateIterator()); Iter; Iter++)
 				{
 					OutResults.Add(*Iter, EPerforceState::ReadOnly);
 				}
 
+				InChangelist = InCommand.Changelist;
 				OutChangelist = ChangeList;
 			}
 		}
@@ -488,10 +519,17 @@ bool FPerforceCheckInWorker::UpdateStates() const
 	bool bUpdatedStates = UpdateCachedStates(OutResults);
 	bool bUpdatedChangelistStates = false;
 
-	if (!OutChangelist.IsDefault())
+	if(!OutChangelist.IsDefault()) // e.g. operation succeeded
 	{
+		// Delete changelist, whether its a temporary one or not
 		FPerforceSourceControlModule& PerforceSourceControl = FPerforceSourceControlModule::Get();
 		bUpdatedChangelistStates = PerforceSourceControl.GetProvider().RemoveChangelistFromCache(OutChangelist);
+
+		// If it's a temporary one, then remove the submitted files from the default changelist
+		if (InChangelist.IsDefault())
+		{
+			bUpdatedChangelistStates = RemoveFilesFromChangelist(OutResults, InChangelist);
+		}
 	}
 
 	return (bUpdatedStates || bUpdatedChangelistStates);
@@ -612,23 +650,6 @@ bool FPerforceRevertWorker::Execute(FPerforceSourceControlCommand& InCommand)
 	return InCommand.bCommandSuccessful;
 }
 
-bool RemoveFilesFromChangelist(const TMap<FString, EPerforceState::Type>& Results, TSharedRef<FPerforceSourceControlChangelistState, ESPMode::ThreadSafe>& ChangelistState)
-{
-	return ChangelistState->Files.RemoveAll([&Results](FSourceControlStateRef& State) -> bool
-		{
-			return Algo::AnyOf(Results, [&State](auto& Result) {
-				return State->GetFilename() == Result.Key;
-				});
-		}) > 0;
-}
-
-bool RemoveFilesFromChangelist(const TMap<FString, EPerforceState::Type>& Results, const FPerforceSourceControlChangelist& Changelist)
-{
-	FPerforceSourceControlModule& PerforceSourceControl = FPerforceSourceControlModule::Get();
-	TSharedRef<FPerforceSourceControlChangelistState, ESPMode::ThreadSafe> ChangelistState = PerforceSourceControl.GetProvider().GetStateInternal(Changelist);
-	return RemoveFilesFromChangelist(Results, ChangelistState);
-}
-
 bool FPerforceRevertWorker::UpdateStates() const
 {
 	bool bUpdatedCachedStates = UpdateCachedStates(OutResults);
@@ -2017,15 +2038,14 @@ FName FPerforceDeleteChangelistWorker::GetName() const
 
 bool FPerforceDeleteChangelistWorker::Execute(class FPerforceSourceControlCommand& InCommand)
 {
+	FScopedPerforceConnection ScopedConnection(InCommand);
+
 	// Can't delete the default changelist
 	if (InCommand.Changelist.IsDefault())
 	{
 		InCommand.bCommandSuccessful = false;
-		return false;
 	}
-
-	FScopedPerforceConnection ScopedConnection(InCommand);
-	if (!InCommand.IsCanceled() && ScopedConnection.IsValid())
+	else if (!InCommand.IsCanceled() && ScopedConnection.IsValid())
 	{
 		FPerforceConnection& Connection = ScopedConnection.GetConnection();
 		check(InCommand.Operation->GetName() == GetName());
@@ -2078,13 +2098,22 @@ bool FPerforceEditChangelistWorker::Execute(class FPerforceSourceControlCommand&
 		check(InCommand.Operation->GetName() == GetName());
 		TSharedRef<FEditChangelist, ESPMode::ThreadSafe> Operation = StaticCastSharedRef<FEditChangelist>(InCommand.Operation);
 
-		int32 ChangelistNumber = Connection.EditPendingChangelist(Operation->GetDescription(), InCommand.Changelist.ToInt(), FOnIsCancelled::CreateRaw(&InCommand, &FPerforceSourceControlCommand::IsCanceled), InCommand.ResultInfo.ErrorMessages);
+		int32 ChangelistNumber = -1;
 
-		InCommand.bCommandSuccessful = (ChangelistNumber == InCommand.Changelist.ToInt());
+		if (InCommand.Changelist.IsDefault())
+		{
+			ChangelistNumber = Connection.CreatePendingChangelist(Operation->GetDescription(), InCommand.Files, FOnIsCancelled::CreateRaw(&InCommand, &FPerforceSourceControlCommand::IsCanceled), InCommand.ResultInfo.ErrorMessages);
+		}
+		else
+		{
+			ChangelistNumber = Connection.EditPendingChangelist(Operation->GetDescription(), InCommand.Changelist.ToInt(), FOnIsCancelled::CreateRaw(&InCommand, &FPerforceSourceControlCommand::IsCanceled), InCommand.ResultInfo.ErrorMessages);
+		}
+
+		InCommand.bCommandSuccessful = (ChangelistNumber == InCommand.Changelist.ToInt() || (ChangelistNumber >= 0 && InCommand.Changelist.IsDefault()));
 
 		if (InCommand.bCommandSuccessful)
 		{
-			EditedChangelist = InCommand.Changelist;
+			EditedChangelist = FPerforceSourceControlChangelist(ChangelistNumber);
 			EditedDescription = Operation->GetDescription();
 		}
 	}
@@ -2096,12 +2125,14 @@ bool FPerforceEditChangelistWorker::UpdateStates() const
 {
 	FPerforceSourceControlModule& PerforceSourceControl = FPerforceSourceControlModule::Get();
 	TSharedRef<FPerforceSourceControlChangelistState, ESPMode::ThreadSafe> EditedChangelistState = PerforceSourceControl.GetProvider().GetStateInternal(EditedChangelist);
+	// TODO: update similar to NewChangelist when/if we support files in edit/new changelists.
 	EditedChangelistState->Description = EditedDescription.ToString();
+	EditedChangelistState->Changelist = EditedChangelist;
+	EditedChangelistState->TimeStamp = FDateTime::Now();
 
 	return true;
 }
 
-
 FName FPerforceRevertUnchangedWorker::GetName() const
 {
 	return "RevertUnchanged";
diff --git a/Engine/Plugins/Developer/PerforceSourceControl/Source/PerforceSourceControl/Private/PerforceSourceControlOperations.h b/Engine/Plugins/Developer/PerforceSourceControl/Source/PerforceSourceControl/Private/PerforceSourceControlOperations.h
--- a/Engine/Plugins/Developer/PerforceSourceControl/Source/PerforceSourceControl/Private/PerforceSourceControlOperations.h
+++ b/Engine/Plugins/Developer/PerforceSourceControl/Source/PerforceSourceControl/Private/PerforceSourceControlOperations.h
@@ -46,6 +46,9 @@ public:
 	/** Map of filenames to perforce state */
 	TMap<FString, EPerforceState::Type> OutResults;
 
+	/** Changelist we asked to submit */
+	FPerforceSourceControlChangelist InChangelist;
+
 	/** Changelist we submitted */
 	FPerforceSourceControlChangelist OutChangelist;
 };
diff --git a/Engine/Source/Editor/SourceControlWindows/Private/SSourceControlChangelists.cpp b/Engine/Source/Editor/SourceControlWindows/Private/SSourceControlChangelists.cpp
--- a/Engine/Source/Editor/SourceControlWindows/Private/SSourceControlChangelists.cpp
+++ b/Engine/Source/Editor/SourceControlWindows/Private/SSourceControlChangelists.cpp
@@ -661,6 +661,11 @@ void SSourceControlChangelistsWidget::OnRevertUnchanged()
 	SourceControlProvider.Execute(RevertUnchangedOperation, GetChangelistFromSelection(), GetSelectedFiles());
 }
 
+bool SSourceControlChangelistsWidget::CanRevertUnchanged()
+{
+	return GetSelectedFiles().Num() > 0 || (GetCurrentChangelistState() && GetCurrentChangelistState()->GetFilesStates().Num() > 0);
+}
+
 void SSourceControlChangelistsWidget::OnRevert()
 {
 	FText DialogText;
@@ -691,6 +696,11 @@ void SSourceControlChangelistsWidget::OnRevert()
 	SourceControlProvider.Execute(RevertOperation, GetChangelistFromSelection(), GetSelectedFiles());
 }
 
+bool SSourceControlChangelistsWidget::CanRevert()
+{
+	return GetSelectedFiles().Num() > 0 || (GetCurrentChangelistState() && GetCurrentChangelistState()->GetFilesStates().Num() > 0);
+}
+
 void SSourceControlChangelistsWidget::OnShelve()
 {
 	ISourceControlProvider& SourceControlProvider = ISourceControlModule::Get().GetProvider();
@@ -714,9 +724,9 @@ void SSourceControlChangelistsWidget::OnDeleteShelvedFiles()
 
 void SSourceControlChangelistsWidget::OnSubmitChangelist()
 {
-	FSourceControlChangelistPtr Changelist = GetCurrentChangelist();
+	FSourceControlChangelistStatePtr ChangelistState = GetCurrentChangelistState();
 	
-	if (!Changelist)
+	if (!ChangelistState)
 	{
 		return;
 	}
@@ -724,14 +734,36 @@ void SSourceControlChangelistsWidget::OnSubmitChangelist()
 	const FText DialogText = LOCTEXT("SourceControl_ConfirmSubmit", "Are you sure you want to submit this changelist?");
 	const FText DialogTitle = LOCTEXT("SourceControl_ConfirmSubmit_Title", "Confirm changelist submit");
 
-	EAppReturnType::Type UserConfirmation = FMessageDialog::Open(EAppMsgType::OkCancel, EAppReturnType::Ok, DialogText, & DialogTitle);
+	EAppReturnType::Type UserConfirmation = FMessageDialog::Open(EAppMsgType::OkCancel, EAppReturnType::Ok, DialogText, &DialogTitle);
 
 	if (UserConfirmation == EAppReturnType::Ok)
 	{
+		FText ChangelistDescription = ChangelistState->GetDescriptionText();
+		const bool bAskForChangelistDescription = (ChangelistDescription.IsEmptyOrWhitespace());
+
+		if (bAskForChangelistDescription)
+		{
+			bool bOk = GetChangelistDescription(
+				nullptr,
+				LOCTEXT("SourceControl.Changelist.Submit.Description.Title", "Enter a submit description..."),
+				LOCTEXT("SourceControl.Changelist.Submit.Description.Label", "Enter a description for the submit:"),
+				ChangelistDescription);
+
+			if (!bOk)
+			{
+				return; // Abort submit
+			}
+		}
+
 		ISourceControlProvider& SourceControlProvider = ISourceControlModule::Get().GetProvider();
 		auto SubmitChangelistOperation = ISourceControlOperation::Create<FCheckIn>();
 
-		SourceControlProvider.Execute(SubmitChangelistOperation, Changelist);
+		if (bAskForChangelistDescription)
+		{
+			SubmitChangelistOperation->SetDescription(ChangelistDescription);
+		}
+
+		SourceControlProvider.Execute(SubmitChangelistOperation, ChangelistState->GetChangelist());
 		Refresh();
 	}
 }
@@ -739,7 +771,7 @@ void SSourceControlChangelistsWidget::OnSubmitChangelist()
 bool SSourceControlChangelistsWidget::CanSubmitChangelist()
 {
 	FSourceControlChangelistStatePtr Changelist = GetCurrentChangelistState();
-	return Changelist != nullptr && Changelist->GetFilesStates().Num() >= 0 && Changelist->GetShelvedFilesStates().Num() == 0;
+	return Changelist != nullptr && Changelist->GetFilesStates().Num() > 0 && Changelist->GetShelvedFilesStates().Num() == 0;
 }
 
 void SSourceControlChangelistsWidget::OnLocateFile()
@@ -833,15 +865,19 @@ TSharedPtr<SWidget> SSourceControlChangelistsWidget::OnOpenContextMenu()
 	if (bHasSelectedChangelist || bHasSelectedFiles)
 	{
 		Section.AddMenuEntry("RevertUnchanged", LOCTEXT("SourceControl_RevertUnchanged", "Revert Unchanged"), LOCTEXT("SourceControl_Revert_Unchanged_Tooltip", "Reverts unchanged files & changelists"), FSlateIcon(),
-			FUIAction(FExecuteAction::CreateSP(this, &SSourceControlChangelistsWidget::OnRevertUnchanged)));
+			FUIAction(
+				FExecuteAction::CreateSP(this, &SSourceControlChangelistsWidget::OnRevertUnchanged),
+				FCanExecuteAction::CreateSP(this, &SSourceControlChangelistsWidget::CanRevertUnchanged)));
 
 		Section.AddMenuEntry("Revert", LOCTEXT("SourceControl_Revert", "Revert Files"), LOCTEXT("SourceControl_Revert_Tooltip", "Reverts all files in the changelist or from the selection"), FSlateIcon(),
-			FUIAction(FExecuteAction::CreateSP(this, &SSourceControlChangelistsWidget::OnRevert)));
+			FUIAction(
+				FExecuteAction::CreateSP(this, &SSourceControlChangelistsWidget::OnRevert),
+				FCanExecuteAction::CreateSP(this, &SSourceControlChangelistsWidget::CanRevert)));
 	}
 
 	if(bHasSelectedFiles || bHasSelectedShelvedFiles || (bHasSelectedChangelist && (GetCurrentChangelistState()->GetFilesStates().Num() > 0 || GetCurrentChangelistState()->GetShelvedFilesStates().Num() > 0)))
 	{
-		Section.AddSeparator("Shelve");
+		Section.AddSeparator("ShelveSeparator");
 	}
 
 	if(bHasSelectedFiles || (bHasSelectedChangelist && GetCurrentChangelistState()->GetFilesStates().Num() > 0))
@@ -871,7 +907,7 @@ TSharedPtr<SWidget> SSourceControlChangelistsWidget::OnOpenContextMenu()
 
 	if (bHasEmptySelection || bHasSelectedChangelist)
 	{
-		Section.AddSeparator("Changelists");
+		Section.AddSeparator("ChangelistsSeparator");
 	}
 
 	// This should appear only if we have no selection
@@ -895,18 +931,18 @@ TSharedPtr<SWidget> SSourceControlChangelistsWidget::OnOpenContextMenu()
 	// Files-only operations
 	if(bHasSelectedFiles)
 	{
-		Section.AddSeparator("Files");
+		Section.AddSeparator("FilesSeparator");
 
-		Section.AddMenuEntry("Locate File", LOCTEXT("SourceControl_LocateFile", "Locate File"), LOCTEXT("SourceControl_LocateFile_Tooltip", "Locate File in Project..."), FSlateIcon(),
+		Section.AddMenuEntry("LocateFile", LOCTEXT("SourceControl_LocateFile", "Locate File"), LOCTEXT("SourceControl_LocateFile_Tooltip", "Locate File in Project..."), FSlateIcon(),
 			FUIAction(
 				FExecuteAction::CreateSP(this, &SSourceControlChangelistsWidget::OnLocateFile),
 				FCanExecuteAction::CreateSP(this, &SSourceControlChangelistsWidget::CanLocateFile)));
 
-		Section.AddMenuEntry("Show History", LOCTEXT("SourceControl_ShowHistory", "Show History"), LOCTEXT("SourceControl_ShowHistory_ToolTip", "Show File History From Selection..."), FSlateIcon(),
+		Section.AddMenuEntry("ShowHistory", LOCTEXT("SourceControl_ShowHistory", "Show History"), LOCTEXT("SourceControl_ShowHistory_ToolTip", "Show File History From Selection..."), FSlateIcon(),
 			FUIAction(
 				FExecuteAction::CreateSP(this, &SSourceControlChangelistsWidget::OnShowHistory)));
 
-		Section.AddMenuEntry("Diff Against Local Version", LOCTEXT("SourceControl_DiffAgainstDepot", "Diff Against Depot"), LOCTEXT("SourceControl_DiffAgainstLocal_Tooltip", "Diff local file against depot revision."), FSlateIcon(),
+		Section.AddMenuEntry("DiffAgainstLocalVersion", LOCTEXT("SourceControl_DiffAgainstDepot", "Diff Against Depot"), LOCTEXT("SourceControl_DiffAgainstLocal_Tooltip", "Diff local file against depot revision."), FSlateIcon(),
 			FUIAction(
 				FExecuteAction::CreateSP(this, &SSourceControlChangelistsWidget::OnDiffAgainstDepot),
 				FCanExecuteAction::CreateSP(this, &SSourceControlChangelistsWidget::CanDiffAgainstDepot)));
diff --git a/Engine/Source/Editor/SourceControlWindows/Private/SSourceControlChangelists.h b/Engine/Source/Editor/SourceControlWindows/Private/SSourceControlChangelists.h
--- a/Engine/Source/Editor/SourceControlWindows/Private/SSourceControlChangelists.h
+++ b/Engine/Source/Editor/SourceControlWindows/Private/SSourceControlChangelists.h
@@ -124,7 +124,9 @@ private:
 
 	/** Changelist & File operations */
 	void OnRevertUnchanged();
+	bool CanRevertUnchanged();
 	void OnRevert();
+	bool CanRevert();
 	void OnShelve();
 
 	/** Changelist & shelved files operations */
