Added possibility to get asset data from history, with a max fetch size (until we have a better solution)
Changed the way we associate shelved files with their local filename to facilitate mapping (perforce operations)
Implemented what was needed on the SCC side to support diff against shelve
Improved way we display object names for deleted & shelved files

diff --git a/Engine/Plugins/Developer/PerforceSourceControl/Source/PerforceSourceControl/Private/PerforceSourceControlOperations.cpp b/Engine/Plugins/Developer/PerforceSourceControl/Source/PerforceSourceControl/Private/PerforceSourceControlOperations.cpp
--- a/Engine/Plugins/Developer/PerforceSourceControl/Source/PerforceSourceControl/Private/PerforceSourceControlOperations.cpp
+++ b/Engine/Plugins/Developer/PerforceSourceControl/Source/PerforceSourceControl/Private/PerforceSourceControlOperations.cpp
@@ -1266,7 +1266,7 @@ static const FString& FindWorkspaceFile(const TArray<FPerforceSourceControlState
 	return InDepotFile;
 }
 
-static void ParseHistoryResults(const FP4RecordSet& InRecords, const TArray<FPerforceSourceControlState>& InStates, FPerforceUpdateStatusWorker::FHistoryMap& OutHistory)
+static void ParseHistoryResults(const FP4RecordSet& InRecords, const TArray<FPerforceSourceControlState>& InStates, FPerforceFileHistoryMap& OutHistory)
 {
 	if (InRecords.Num() > 0)
 	{
@@ -1375,6 +1375,30 @@ static void ParseHistoryResults(const FP4RecordSet& InRecords, const TArray<FPer
 	}
 }
 
+static bool GetFileHistory(FPerforceConnection& Connection, FPerforceSourceControlCommand& InCommand, const TArray<FString>& InFiles, TArray<FPerforceSourceControlState>& OutStates, FPerforceFileHistoryMap& OutHistory)
+{
+	TArray<FString> Parameters;
+	FP4RecordSet Records;
+	// disregard non-contributory integrations
+	Parameters.Add(TEXT("-s"));
+	//include branching history
+	Parameters.Add(TEXT("-i"));
+	//include truncated change list descriptions
+	Parameters.Add(TEXT("-L"));
+	//include time stamps
+	Parameters.Add(TEXT("-t"));
+	//limit to last 100 changes
+	Parameters.Add(TEXT("-m 100"));
+	Parameters.Append(InFiles);
+	InCommand.bCommandSuccessful &= Connection.RunCommand(TEXT("filelog"), Parameters, Records, InCommand.ResultInfo.ErrorMessages, FOnIsCancelled::CreateRaw(&InCommand, &FPerforceSourceControlCommand::IsCanceled), InCommand.bConnectionDropped);
+	ParseHistoryResults(Records, OutStates, OutHistory);
+	RemoveRedundantErrors(InCommand, TEXT(" - no such file(s)."));
+	RemoveRedundantErrors(InCommand, TEXT(" - file(s) not on client"));
+	RemoveRedundantErrors(InCommand, TEXT("' is not under client's root '"));
+
+	return InCommand.bCommandSuccessful;
+}
+
 static void ParseDiffResults(const FP4RecordSet& InRecords, TArray<FString>& OutModifiedFiles)
 {
 	if (InRecords.Num() > 0)
@@ -1501,26 +1525,11 @@ bool FPerforceUpdateStatusWorker::Execute(FPerforceSourceControlCommand& InComma
 		check(InCommand.Operation->GetName() == GetName());
 		TSharedRef<FUpdateStatus, ESPMode::ThreadSafe> Operation = StaticCastSharedRef<FUpdateStatus>(InCommand.Operation);
 
+		bForceQuiet = Operation->ShouldBeQuiet();
+
 		if(Operation->ShouldUpdateHistory())
 		{
-			TArray<FString> Parameters;
-			FP4RecordSet Records;
-			// disregard non-contributory integrations
-			Parameters.Add(TEXT("-s"));
-			//include branching history
-			Parameters.Add(TEXT("-i"));
-			//include truncated change list descriptions
-			Parameters.Add(TEXT("-L"));
-			//include time stamps
-			Parameters.Add(TEXT("-t"));
-			//limit to last 100 changes
-			Parameters.Add(TEXT("-m 100"));
-			Parameters.Append(InCommand.Files);
-			InCommand.bCommandSuccessful &= Connection.RunCommand(TEXT("filelog"), Parameters, Records, InCommand.ResultInfo.ErrorMessages, FOnIsCancelled::CreateRaw(&InCommand, &FPerforceSourceControlCommand::IsCanceled), InCommand.bConnectionDropped);
-			ParseHistoryResults(Records, OutStates, OutHistory);
-			RemoveRedundantErrors(InCommand, TEXT(" - no such file(s)."));
-			RemoveRedundantErrors(InCommand, TEXT(" - file(s) not on client"));
-			RemoveRedundantErrors(InCommand, TEXT("' is not under client's root '"));
+			GetFileHistory(Connection, InCommand, InCommand.Files, OutStates, OutHistory);
 		}
 
 		if(Operation->ShouldGetOpenedOnly())
@@ -1593,7 +1602,7 @@ bool FPerforceUpdateStatusWorker::UpdateStates() const
 	bUpdated |= UpdateCachedStates(OutStateMap);
 
 	// add history, if any
-	for(FHistoryMap::TConstIterator It(OutHistory); It; ++It)
+	for(FPerforceFileHistoryMap::TConstIterator It(OutHistory); It; ++It)
 	{
 		TSharedRef<FPerforceSourceControlState, ESPMode::ThreadSafe> State = PerforceSourceControl.GetProvider().GetStateInternal(It.Key());
 		const TArray< TSharedRef<FPerforceSourceControlRevision, ESPMode::ThreadSafe> >& History = It.Value();
@@ -1612,7 +1621,7 @@ bool FPerforceUpdateStatusWorker::UpdateStates() const
 		bUpdated = true;
 	}
 
-	return bUpdated;
+	return !bForceQuiet && bUpdated;
 }
 
 FName FPerforceGetWorkspacesWorker::GetName() const
@@ -1665,6 +1674,49 @@ static bool GetOpenedFilesInChangelist(FPerforceConnection& Connection, FPerforc
 	return InCommand.bCommandSuccessful;
 }
 
+static void ParseWhereResults(FP4RecordSet& InRecords, TMap<FString, FString>& DepotToFileMap)
+{
+	for (int32 Index = 0; Index < InRecords.Num(); ++Index)
+	{
+		const FP4Record& Record = InRecords[Index];
+
+		FString DepotFile = Record(TEXT("depotFile"));
+		FString ClientFile = Record(TEXT("path")).Replace(TEXT("\\"), TEXT("/"));
+
+		if (DepotFile.Len() > 0 && ClientFile.Len() > 0)
+		{
+			DepotToFileMap.Emplace(DepotFile, ClientFile);
+		}
+	}
+}
+
+static bool GetDepotFileToLocalFileMap(FPerforceConnection& Connection, FPerforceSourceControlCommand& InCommand, const TMap<FString, EPerforceState::Type>& InDepotFiles, TMap<FString, FString>& OutDepotToLocalMap)
+{
+	if (InDepotFiles.Num() == 0)
+	{
+		return true;
+	}
+
+	TArray<FString> Parameters;
+	for (TMap<FString, EPerforceState::Type>::TConstIterator It(InDepotFiles); It; ++It)
+	{
+		Parameters.Add(It.Key());
+	}
+
+	FP4RecordSet Records;
+	Connection.RunCommand(TEXT("where"), Parameters, Records, InCommand.ResultInfo.ErrorMessages, FOnIsCancelled::CreateRaw(&InCommand, &FPerforceSourceControlCommand::IsCanceled), InCommand.bConnectionDropped);
+
+	if (InCommand.ResultInfo.ErrorMessages.Num() == 0)
+	{
+		ParseWhereResults(Records, OutDepotToLocalMap);
+		return true;
+	}
+	else
+	{
+		return false;
+	}
+}
+
 bool FPerforceGetPendingChangelistsWorker::Execute(FPerforceSourceControlCommand& InCommand)
 {
 	FScopedPerforceConnection ScopedConnection(InCommand);
@@ -1729,6 +1781,7 @@ bool FPerforceGetPendingChangelistsWorker::Execute(FPerforceSourceControlCommand
 				if (!ChangelistState.bHasShelvedFiles)
 				{
 					OutCLShelvedFilesStates.Emplace(); // Add empty list
+					OutCLShelvedFilesMap.Emplace(); // Add empty list
 					continue;
 				}
 
@@ -1745,6 +1798,9 @@ bool FPerforceGetPendingChangelistsWorker::Execute(FPerforceSourceControlCommand
 				{
 					TMap<FString, EPerforceState::Type>& OutShelvedStateMap = OutCLShelvedFilesStates.Emplace_GetRef();
 					ParseShelvedChangelistResults(Records, OutShelvedStateMap);
+
+					TMap<FString, FString>& OutShelvedFileMap = OutCLShelvedFilesMap.Emplace_GetRef();
+					GetDepotFileToLocalFileMap(Connection, InCommand, OutShelvedStateMap, OutShelvedFileMap);
 				}
 			}
 		}
@@ -1760,22 +1816,45 @@ bool FPerforceGetPendingChangelistsWorker::Execute(FPerforceSourceControlCommand
 	return InCommand.bCommandSuccessful;
 }
 
-static bool AddShelvedFilesToChangelist(const TMap<FString, EPerforceState::Type>& FilesToAdd, TSharedRef<FPerforceSourceControlChangelistState, ESPMode::ThreadSafe>& ChangelistState, const FDateTime* TimeStamp = nullptr)
+static bool AddShelvedFilesToChangelist(const TMap<FString, EPerforceState::Type>& FilesToAdd, const TMap<FString, FString>& DepotToFileMap, TSharedRef<FPerforceSourceControlChangelistState, ESPMode::ThreadSafe>& ChangelistState, const FDateTime* TimeStamp = nullptr)
 {
 	FPerforceSourceControlModule& PerforceSourceControl = FPerforceSourceControlModule::Get();
 	const FDateTime Now = (TimeStamp ? *TimeStamp : FDateTime::Now());
 
 	for (TMap<FString, EPerforceState::Type>::TConstIterator It(FilesToAdd); It; ++It)
 	{
-		int32 Index = Algo::IndexOfByPredicate(ChangelistState->ShelvedFiles, [&It](const FSourceControlStateRef& ShelvedFile) {
-				return ShelvedFile->GetFilename() == It.Key();
+		FString ItDepotFilename = It.Key();
+		FString ItFilename = ItDepotFilename;
+
+		if (DepotToFileMap.Contains(ItDepotFilename))
+		{
+			ItFilename = DepotToFileMap[It.Key()];
+		}
+
+		int32 Index = Algo::IndexOfByPredicate(ChangelistState->ShelvedFiles, [&ItFilename](const FSourceControlStateRef& ShelvedFile) {
+				return ShelvedFile->GetFilename() == ItFilename;
 			});
 
 		if (Index < 0)
 		{
 			// Create new entry
+			TSharedRef<FPerforceSourceControlState, ESPMode::ThreadSafe> ShelvedFileState = MakeShareable(new FPerforceSourceControlState(ItFilename));
+			ShelvedFileState->DepotFilename = ItDepotFilename;
+
+			// Add revision to be able to fetch the shelved file, if it's not marked for deletion.
+			if (It.Value() != EPerforceState::MarkedForDelete)
+			{
+				TSharedRef<FPerforceSourceControlRevision, ESPMode::ThreadSafe> ShelvedRevision = MakeShareable(new FPerforceSourceControlRevision());
+				ShelvedRevision->FileName = ShelvedFileState->DepotFilename;
+				ShelvedRevision->ChangelistNumber = StaticCastSharedRef<FPerforceSourceControlChangelist>(ChangelistState->GetChangelist())->ToInt();
+				ShelvedRevision->bIsShelve = true;
+
+				ShelvedFileState->History.Add(ShelvedRevision);
+			}
+
+			// Add to shelved files
 			Index = ChangelistState->ShelvedFiles.Num();
-			ChangelistState->ShelvedFiles.Emplace(MakeShareable(new FPerforceSourceControlState(It.Key())));
+			ChangelistState->ShelvedFiles.Add(ShelvedFileState);
 		}
 
 		TSharedRef<FPerforceSourceControlState, ESPMode::ThreadSafe> FileState = StaticCastSharedRef<FPerforceSourceControlState>(ChangelistState->ShelvedFiles[Index]);
@@ -1787,11 +1866,11 @@ static bool AddShelvedFilesToChangelist(const TMap<FString, EPerforceState::Type
 	return FilesToAdd.Num() > 0;
 }
 
-static bool AddShelvedFilesToChangelist(const TMap<FString, EPerforceState::Type>& FilesToAdd, const FPerforceSourceControlChangelist& Changelist)
+static bool AddShelvedFilesToChangelist(const TMap<FString, EPerforceState::Type>& FilesToAdd, const TMap<FString, FString>& DepotToFileMap, const FPerforceSourceControlChangelist& Changelist)
 {
 	FPerforceSourceControlModule& PerforceSourceControl = FPerforceSourceControlModule::Get();
 	TSharedRef<FPerforceSourceControlChangelistState, ESPMode::ThreadSafe> ChangelistState = PerforceSourceControl.GetProvider().GetStateInternal(Changelist);
-	return AddShelvedFilesToChangelist(FilesToAdd, ChangelistState);
+	return AddShelvedFilesToChangelist(FilesToAdd, DepotToFileMap, ChangelistState);
 }
 
 bool FPerforceGetPendingChangelistsWorker::UpdateStates() const
@@ -1829,7 +1908,7 @@ bool FPerforceGetPendingChangelistsWorker::UpdateStates() const
 		if(bUpdateShelvedFiles)
 		{
 			ChangelistState->ShelvedFiles.Reset(OutCLShelvedFilesStates[StatusIndex].Num());
-			AddShelvedFilesToChangelist(OutCLShelvedFilesStates[StatusIndex], ChangelistState, &Now);
+			AddShelvedFilesToChangelist(OutCLShelvedFilesStates[StatusIndex], OutCLShelvedFilesMap[StatusIndex], ChangelistState, &Now);
 		}
 	}
 
@@ -2243,6 +2322,9 @@ bool FPerforceShelveWorker::Execute(class FPerforceSourceControlCommand& InComma
 		{
 			ChangelistToUpdate = InCommand.Changelist;
 			ParseShelvedResults(Records, OutResults);
+
+			// Build depot to file mapping
+			GetDepotFileToLocalFileMap(Connection, InCommand, OutResults, OutFileMap);
 		}
 	}
 
@@ -2251,7 +2333,7 @@ bool FPerforceShelveWorker::Execute(class FPerforceSourceControlCommand& InComma
 
 bool FPerforceShelveWorker::UpdateStates() const
 {
-	return (OutResults.Num() > 0 && AddShelvedFilesToChangelist(OutResults, ChangelistToUpdate));
+	return (OutResults.Num() > 0 && AddShelvedFilesToChangelist(OutResults, OutFileMap, ChangelistToUpdate));
 }
 
 
diff --git a/Engine/Plugins/Developer/PerforceSourceControl/Source/PerforceSourceControl/Private/PerforceSourceControlOperations.h b/Engine/Plugins/Developer/PerforceSourceControl/Source/PerforceSourceControl/Private/PerforceSourceControlOperations.h
--- a/Engine/Plugins/Developer/PerforceSourceControl/Source/PerforceSourceControl/Private/PerforceSourceControlOperations.h
+++ b/Engine/Plugins/Developer/PerforceSourceControl/Source/PerforceSourceControl/Private/PerforceSourceControlOperations.h
@@ -8,6 +8,7 @@
 #include "PerforceSourceControlChangelistState.h"
 
 class FPerforceSourceControlRevision;
+typedef TMap<FString, TArray< TSharedRef<FPerforceSourceControlRevision, ESPMode::ThreadSafe> > > FPerforceFileHistoryMap;
 
 class FPerforceConnectWorker : public IPerforceSourceControlWorker
 {
@@ -129,11 +130,13 @@ public:
 	TMap<FString, EPerforceState::Type> OutStateMap;
 
 	/** Map of filenames to history */
-	typedef TMap<FString, TArray< TSharedRef<FPerforceSourceControlRevision, ESPMode::ThreadSafe> > > FHistoryMap;
-	FHistoryMap OutHistory;
+	FPerforceFileHistoryMap OutHistory;
 
 	/** Map of filenames to modified flag */
 	TArray<FString> OutModifiedFiles;
+
+	/** Override on status update return */
+	bool bForceQuiet = false;
 };
 
 class FPerforceGetWorkspacesWorker : public IPerforceSourceControlWorker
@@ -160,6 +163,7 @@ public:
 	TArray<FPerforceSourceControlChangelistState> OutChangelistsStates;
 	TArray<TArray<FPerforceSourceControlState>> OutCLFilesStates;
 	TArray<TMap<FString, EPerforceState::Type>> OutCLShelvedFilesStates;
+	TArray<TMap<FString, FString>> OutCLShelvedFilesMap;
 };
 
 class FPerforceCopyWorker : public IPerforceSourceControlWorker
@@ -287,6 +291,9 @@ protected:
 	/** Map of filenames to perforce state */
 	TMap<FString, EPerforceState::Type> OutResults;
 
+	/** Map depot filenames to local file */
+	TMap<FString, FString> OutFileMap;
+
 	/** Changelist to be updated */
 	FPerforceSourceControlChangelist ChangelistToUpdate;
 };
diff --git a/Engine/Plugins/Developer/PerforceSourceControl/Source/PerforceSourceControl/Private/PerforceSourceControlRevision.cpp b/Engine/Plugins/Developer/PerforceSourceControl/Source/PerforceSourceControl/Private/PerforceSourceControlRevision.cpp
--- a/Engine/Plugins/Developer/PerforceSourceControl/Source/PerforceSourceControl/Private/PerforceSourceControlRevision.cpp
+++ b/Engine/Plugins/Developer/PerforceSourceControl/Source/PerforceSourceControl/Private/PerforceSourceControlRevision.cpp
@@ -28,7 +28,20 @@ bool FPerforceSourceControlRevision::Get( FString& InOutFilename ) const
 		Parameters.Add(TEXT("-q"));
 
 		// Make temp filename to 'print' to
-		FString RevString = (RevisionNumber < 0) ? TEXT("head") : FString::Printf(TEXT("%d"), RevisionNumber);
+		FString RevString;
+		FString RevParam;
+		
+		if (bIsShelve)
+		{
+			RevString = FString::Printf(TEXT("%d"), ChangelistNumber);
+			RevParam = TEXT("@=") + RevString;
+		}
+		else
+		{
+			RevString = (RevisionNumber < 0) ? TEXT("head") : FString::Printf(TEXT("%d"), RevisionNumber);
+			RevParam = TEXT("#") + RevString;
+		}
+		
 		FString AbsoluteFileName;
 		if(InOutFilename.Len() > 0)
 		{
@@ -43,7 +56,7 @@ bool FPerforceSourceControlRevision::Get( FString& InOutFilename ) const
 
 		// output to file
 		Parameters.Add(FString("-o") + AbsoluteFileName);
-		Parameters.Add(FileName + TEXT("#") + RevString);
+		Parameters.Add(FileName + RevParam);
 		
 		bCommandOK = Connection.RunCommand(TEXT("print"), Parameters, Records, ErrorMessages, FOnIsCancelled(), bConnectionDropped);
 		if(bCommandOK)
diff --git a/Engine/Plugins/Developer/PerforceSourceControl/Source/PerforceSourceControl/Private/PerforceSourceControlRevision.h b/Engine/Plugins/Developer/PerforceSourceControl/Source/PerforceSourceControl/Private/PerforceSourceControlRevision.h
--- a/Engine/Plugins/Developer/PerforceSourceControl/Source/PerforceSourceControl/Private/PerforceSourceControlRevision.h
+++ b/Engine/Plugins/Developer/PerforceSourceControl/Source/PerforceSourceControl/Private/PerforceSourceControlRevision.h
@@ -13,6 +13,7 @@ public:
 		, Date(0)
 		, ChangelistNumber(0)
 		, FileSize(0)
+		, bIsShelve(false)
 	{
 	}
 
@@ -65,4 +66,7 @@ public:
 
 	/** The size of the change */
 	int32 FileSize;
+
+	/** Whether this reprensents a revision bound to a shelved file in a changelist */
+	bool bIsShelve;
 };
diff --git a/Engine/Source/Developer/AssetTools/Private/AssetTools.cpp b/Engine/Source/Developer/AssetTools/Private/AssetTools.cpp
--- a/Engine/Source/Developer/AssetTools/Private/AssetTools.cpp
+++ b/Engine/Source/Developer/AssetTools/Private/AssetTools.cpp
@@ -1469,6 +1469,13 @@ void UAssetToolsImpl::DiffAgainstDepot( UObject* InObject, const FString& InPack
 					{
 						// Grab the old asset from that old package
 						UObject* OldObject = FindObject<UObject>(TempPackage, *InPackageName);
+
+						// Recovery for package names that don't match
+						if (OldObject == nullptr)
+						{
+							OldObject = TempPackage->FindAssetInPackage();
+						}
+
 						if(OldObject != nullptr)
 						{
 							/* Set the revision information*/
diff --git a/Engine/Source/Developer/SourceControl/Private/SourceControlHelpers.cpp b/Engine/Source/Developer/SourceControl/Private/SourceControlHelpers.cpp
--- a/Engine/Source/Developer/SourceControl/Private/SourceControlHelpers.cpp
+++ b/Engine/Source/Developer/SourceControl/Private/SourceControlHelpers.cpp
@@ -1445,6 +1445,10 @@ bool USourceControlHelpers::GetAssetData(const FString & InFileName, const FStri
 {
 	const bool bGetDependencies = (OutDependencies != nullptr);
 	OutAssets.Reset();
+	if (bGetDependencies)
+	{
+		OutDependencies->Reset();
+	}
 
 	// Try the registry first
 	FAssetRegistryModule& AssetRegistryModule = FModuleManager::LoadModuleChecked<FAssetRegistryModule>(TEXT("AssetRegistry"));
@@ -1484,6 +1488,65 @@ bool USourceControlHelpers::GetAssetData(const FString & InFileName, const FStri
 	return OutAssets.Num() > 0;
 }
 
+bool USourceControlHelpers::GetAssetDataFromFileHistory(const FString& InFileName, TArray<FAssetData>& OutAssets, TArray<FName>* OutDependencies, int64 MaxFetchSize)
+{
+	OutAssets.Reset();
+
+	if (OutDependencies)
+	{
+		OutDependencies->Reset();
+	}
+
+	ISourceControlProvider& SourceControlProvider = ISourceControlModule::Get().GetProvider();
+	// Get the SCC state
+	FSourceControlStatePtr SourceControlState = SourceControlProvider.GetState(InFileName, EStateCacheUsage::Use);
+	if (SourceControlState.IsValid())
+	{
+		return GetAssetDataFromFileHistory(SourceControlState, OutAssets, OutDependencies, MaxFetchSize);
+	}
+	else
+	{
+		return false;
+	}
+}
+
+bool USourceControlHelpers::GetAssetDataFromFileHistory(FSourceControlStatePtr InSourceControlState, TArray<FAssetData>& OutAssets, TArray<FName>* OutDependencies /* = nullptr */, int64 MaxFetchSize /* = -1 */)
+{
+	check(InSourceControlState.IsValid());
+	OutAssets.Reset();
+
+	if (OutDependencies)
+	{
+		OutDependencies->Reset();
+	}
+
+	// This code is similar to what's done in UAssetToolsImpl::DiffAgainstDepot but we'll force it quiet to prevent recursion issues
+	if (InSourceControlState->GetHistorySize() == 0)
+	{
+		ISourceControlProvider& SourceControlProvider = ISourceControlModule::Get().GetProvider();
+		TSharedRef<FUpdateStatus, ESPMode::ThreadSafe> UpdateStatusOperation = ISourceControlOperation::Create<FUpdateStatus>();
+		UpdateStatusOperation->SetUpdateHistory(true);
+		UpdateStatusOperation->SetQuiet(true);
+		SourceControlProvider.Execute(UpdateStatusOperation, InSourceControlState->GetFilename());
+	}
+
+	if (InSourceControlState->GetHistorySize() > 0)
+	{
+		TSharedPtr<ISourceControlRevision, ESPMode::ThreadSafe> Revision = InSourceControlState->GetHistoryItem(0);
+		check(Revision.IsValid());
+
+		const bool bShouldGetFile = (MaxFetchSize < 0 || MaxFetchSize >(int64)Revision->GetFileSize());
+
+		FString TempFileName;
+		if (bShouldGetFile && Revision->Get(TempFileName))
+		{
+			return GetAssetData(TempFileName, OutAssets, OutDependencies);
+		}
+	}
+
+	return false;
+}
+
 FScopedSourceControl::FScopedSourceControl()
 {
 	ISourceControlModule::Get().GetProvider().Init();
diff --git a/Engine/Source/Developer/SourceControl/Public/SourceControlHelpers.h b/Engine/Source/Developer/SourceControl/Public/SourceControlHelpers.h
--- a/Engine/Source/Developer/SourceControl/Public/SourceControlHelpers.h
+++ b/Engine/Source/Developer/SourceControl/Public/SourceControlHelpers.h
@@ -511,6 +511,18 @@ public:
 	*/
 	static bool GetAssetDataFromPackage(const FString& InPackageName, TArray<FAssetData>& OutAssets, TArray<FName>* OutDependencies = nullptr);
 
+	/**
+	 * Helper function to get the asset data from a file's history
+	 * Note: will not query history to prevent recursion
+	 */
+	static bool GetAssetDataFromFileHistory(const FString& InFileName, TArray<FAssetData>& OutAssets, TArray<FName>* OutDependencies = nullptr, int64 MaxFetchSize = -1);
+
+	/**
+	 * Helper function to get the asset data from a file's history
+	 * Note: will not query history to prevent recursion
+	 */
+	static bool GetAssetDataFromFileHistory(FSourceControlStatePtr InSourceControlState, TArray<FAssetData>& OutAssets, TArray<FName>* OutDependencies = nullptr, int64 MaxFetchSize = -1);
+
 };  // USourceControlHelpers
 
 
diff --git a/Engine/Source/Developer/SourceControl/Public/SourceControlOperations.h b/Engine/Source/Developer/SourceControl/Public/SourceControlOperations.h
--- a/Engine/Source/Developer/SourceControl/Public/SourceControlOperations.h
+++ b/Engine/Source/Developer/SourceControl/Public/SourceControlOperations.h
@@ -218,6 +218,7 @@ public:
 		, bGetOpenedOnly(false)
 		, bUpdateModifiedState(false)
 		, bCheckingAllFiles(false)
+		, bForceQuiet(false)
 	{
 	}
 
@@ -252,6 +253,11 @@ public:
 		bCheckingAllFiles = bInCheckingAllFiles;
 	}
 
+	void SetQuiet(bool bInQuiet)
+	{
+		bForceQuiet = bInQuiet;
+	}
+
 	bool ShouldUpdateHistory() const
 	{
 		return bUpdateHistory;
@@ -272,6 +278,11 @@ public:
 		return bCheckingAllFiles;
 	}
 
+	bool ShouldBeQuiet() const
+	{
+		return bForceQuiet;
+	}
+
 protected:
 	/** Whether to update history */
 	bool bUpdateHistory;
@@ -284,6 +295,9 @@ protected:
 
 	/** Hint that we are intending on checking all files in the project - some providers can optimize for this */
 	bool bCheckingAllFiles;
+
+	/** Controls whether the operation will trigger an update or not */
+	bool bForceQuiet;
 };
 
 /**
diff --git a/Engine/Source/Editor/SourceControlWindows/Private/SSourceControlChangelists.cpp b/Engine/Source/Editor/SourceControlWindows/Private/SSourceControlChangelists.cpp
--- a/Engine/Source/Editor/SourceControlWindows/Private/SSourceControlChangelists.cpp
+++ b/Engine/Source/Editor/SourceControlWindows/Private/SSourceControlChangelists.cpp
@@ -139,7 +139,7 @@ void SSourceControlChangelistsWidget::RequestRefresh()
 		UpdatePendingChangelistsOperation->SetUpdateShelvedFilesStates(true);
 
 		ISourceControlProvider& SourceControlProvider = ISourceControlModule::Get().GetProvider();
-		SourceControlProvider.Execute(UpdatePendingChangelistsOperation, EConcurrency::Asynchronous, FSourceControlOperationComplete::CreateSP(this, &SSourceControlChangelistsWidget::OnChangelistsStatusUpdated));
+		SourceControlProvider.Execute(UpdatePendingChangelistsOperation, EConcurrency::Asynchronous);
 	}
 	else
 	{
@@ -663,7 +663,11 @@ bool SSourceControlChangelistsWidget::CanDiffAgainstDepot()
 
 void SSourceControlChangelistsWidget::OnDiffAgainstWorkspace()
 {
-
+	if (GetSelectedShelvedFiles().Num() > 0)
+	{
+		FSourceControlStateRef FileState = StaticCastSharedPtr<FShelvedFileTreeItem>(TreeView->GetSelectedItems()[0])->FileState;
+		FSourceControlWindows::DiffAgainstShelvedFile(FileState);
+	}
 }
 
 bool SSourceControlChangelistsWidget::CanDiffAgainstWorkspace()
diff --git a/Engine/Source/Editor/SourceControlWindows/Private/SSourceControlCommon.cpp b/Engine/Source/Editor/SourceControlWindows/Private/SSourceControlCommon.cpp
--- a/Engine/Source/Editor/SourceControlWindows/Private/SSourceControlCommon.cpp
+++ b/Engine/Source/Editor/SourceControlWindows/Private/SSourceControlCommon.cpp
@@ -55,24 +55,43 @@ FFileTreeItem::FFileTreeItem(FSourceControlStateRef InFileState, bool bIsShelved
 	// Initialize asset data first
 	FString Filename = FileState->GetFilename();
 
-	USourceControlHelpers::GetAssetData(Filename, Assets);
+	if (!FileState->IsDeleted())
+	{
+		USourceControlHelpers::GetAssetData(Filename, Assets);
+	}
+
+	// For deleted items, the file is not on disk anymore so the only way we can get the asset data is by getting the file from the depot.
+	// For shelved files, if the file still exists locally, it will have been found before, otherwise, the history of the shelved file state will point to the remote version
+	if (FileState->IsDeleted() || (bIsShelvedFile && Assets.Num() == 0))
+	{
+		// At the moment, getting the asset data from non-external assets yields issues with the package path
+		// so we will fall down to our recovery (below) instead
+		if (Filename.Contains("__ExternalActors__"))
+		{
+			const int64 MaxFetchSize = (1 << 20); // 1MB
+			// In the case of shelved "marked for delete", we'll piggy back on the non-shelved file
+			if (bIsShelvedFile && FileState->IsDeleted())
+			{
+				USourceControlHelpers::GetAssetDataFromFileHistory(Filename, Assets, nullptr, MaxFetchSize);
+			}
+			else
+			{
+				USourceControlHelpers::GetAssetDataFromFileHistory(FileState, Assets, nullptr, MaxFetchSize);
+			}
+		}
+	}
 
 	// Initialize display-related members
-	FString TempAssetName = LOCTEXT("SourceControl_DefaultAssetName", "None").ToString();
+	FString TempAssetName = LOCTEXT("SourceControl_DefaultAssetName", "Unavailable").ToString();
 	FString TempAssetPath = Filename;
 	FString TempAssetType = LOCTEXT("SourceControl_DefaultAssetType", "Unknown").ToString();
+	FString TempPackageName = Filename;
 	FColor TempAssetColor = FColor(		// Copied from ContentBrowserCLR.cpp
 		127 + FColor::Red.R / 2,	// Desaturate the colors a bit (GB colors were too.. much)
 		127 + FColor::Red.G / 2,
 		127 + FColor::Red.B / 2,
 		200); // Opacity
 
-	FString TempPackageName = Filename;
-	if (!FPackageName::TryConvertFilenameToLongPackageName(Filename, TempPackageName))
-	{
-		TempPackageName = Filename;
-	}
-
 	if (Assets.Num() > 0)
 	{
 		TempAssetPath = Assets[0].ObjectPath.ToString();
diff --git a/Engine/Source/Editor/SourceControlWindows/Private/SSourceControlHistory.cpp b/Engine/Source/Editor/SourceControlWindows/Private/SSourceControlHistory.cpp
--- a/Engine/Source/Editor/SourceControlWindows/Private/SSourceControlHistory.cpp
+++ b/Engine/Source/Editor/SourceControlWindows/Private/SSourceControlHistory.cpp
@@ -1701,3 +1701,67 @@ bool FSourceControlWindows::DiffAgainstWorkspace(const FString& InFileName)
 
 	return !!SelectedAsset;
 }
+
+bool FSourceControlWindows::DiffAgainstShelvedFile(const FSourceControlStateRef& InFileState)
+{
+	if (InFileState->GetHistorySize() == 0)
+	{
+		return false;
+	}
+
+	TSharedPtr<ISourceControlRevision, ESPMode::ThreadSafe> Revision = InFileState->GetHistoryItem(0);
+	check(Revision.IsValid());
+
+	UObject* SelectedAsset = nullptr;
+
+	FString AssetPackageName;
+	if (FPackageName::TryConvertFilenameToLongPackageName(InFileState->GetFilename(), AssetPackageName))
+	{
+		UPackage* AssetPackage = LoadPackage(nullptr, *AssetPackageName, LOAD_ForDiff | LOAD_DisableCompileOnLoad);
+
+		// grab the asset from the package - we assume asset name matches file name
+		FString AssetName = FPaths::GetBaseFilename(InFileState->GetFilename());
+		SelectedAsset = FindObject<UObject>(AssetPackage, *AssetName);
+
+		if (!SelectedAsset && AssetPackage)
+		{
+			SelectedAsset = AssetPackage->FindAssetInPackage();
+		}
+	}
+
+	if (SelectedAsset)
+	{
+		FString TempFileName;
+		if (Revision->Get(TempFileName))
+		{
+			// Try and load that package
+			UPackage* TempPackage = LoadPackage(nullptr, *TempFileName, LOAD_ForDiff | LOAD_DisableCompileOnLoad);
+			if (TempPackage != nullptr)
+			{
+				// Grab the shelved asset from that package
+				UObject* ShelvedObject = FindObject<UObject>(TempPackage, *SelectedAsset->GetName());
+
+				if (ShelvedObject == nullptr)
+				{
+					ShelvedObject = TempPackage->FindAssetInPackage();
+				}
+
+				if (ShelvedObject != nullptr)
+				{
+					/* Set the revision information*/
+					FRevisionInfo ShelvedRevision;
+					ShelvedRevision.Changelist = Revision->GetCheckInIdentifier();
+					ShelvedRevision.Revision = TEXT("Shelved");
+
+					FRevisionInfo NewRevision;
+					NewRevision.Revision = TEXT("");
+
+					FAssetToolsModule& AssetToolsModule = FModuleManager::LoadModuleChecked<FAssetToolsModule>(TEXT("AssetTools"));
+					AssetToolsModule.Get().DiffAssets(ShelvedObject, SelectedAsset, ShelvedRevision, NewRevision);
+				}
+			}
+		}
+	}
+
+	return !!SelectedAsset;
+}
\ No newline at end of file
diff --git a/Engine/Source/Editor/SourceControlWindows/Public/SourceControlWindows.h b/Engine/Source/Editor/SourceControlWindows/Public/SourceControlWindows.h
--- a/Engine/Source/Editor/SourceControlWindows/Public/SourceControlWindows.h
+++ b/Engine/Source/Editor/SourceControlWindows/Public/SourceControlWindows.h
@@ -103,6 +103,15 @@ public:
 	 */
 	static bool DiffAgainstWorkspace(const FString& InFileName);
 
+	/**
+	 * Displays file diff against shelved file
+	 * 
+	 * @param InFileState SCC file state of the file to diff
+	 * 
+	 * @return true if the diff could be performed
+	 */
+	static bool DiffAgainstShelvedFile(const FSourceControlStateRef& InFileState);
+
 protected:
 	/** Callback for ChoosePackagesToCheckIn(), continues to bring up UI once source control operations are complete */
 	static void ChoosePackagesToCheckInCallback(const FSourceControlOperationRef& InOperation, ECommandResult::Type InResult, FSourceControlWindowsOnCheckInComplete OnCompleteDelegate);
