Fix "shelved files" from default changelist, now has the same behavior as in P4V.
Fix an issue where editing a changelist containing shelved files would fail silently.

diff --git a/Engine/Plugins/Developer/PerforceSourceControl/Source/PerforceSourceControl/Private/PerforceConnection.cpp b/Engine/Plugins/Developer/PerforceSourceControl/Source/PerforceSourceControl/Private/PerforceConnection.cpp
--- a/Engine/Plugins/Developer/PerforceSourceControl/Source/PerforceSourceControl/Private/PerforceConnection.cpp
+++ b/Engine/Plugins/Developer/PerforceSourceControl/Source/PerforceSourceControl/Private/PerforceConnection.cpp
@@ -162,7 +162,7 @@ public:
 
 		for (const auto& Field : Record)
 		{
-			// Skip some fields that aren<t required (as evidenced by the CreateChangelist before
+			// Skip some fields that aren't required (as evidenced by the CreateChangelist before
 			if (Field.Key == TEXT("Date") ||
 				Field.Key == TEXT("Type") ||
 				Field.Key == TEXT("specFormatted") ||
@@ -170,6 +170,11 @@ public:
 			{
 				continue;
 			}
+			else if (Field.Key == TEXT("extraTag0"))
+			{
+				// Changelists with shelved files will have additional tags which will be rejected in the edit
+				break;
+			}
 
 			if (Field.Key == TEXT("Description"))
 			{
@@ -196,7 +201,7 @@ public:
 
 				OutputDesc += TEXT("\t");
 				OutputDesc += Field.Value;
-				OutputDesc += TEXT("\n");
+				OutputDesc += TEXT("\n\n");
 			}
 			else
 			{
@@ -818,7 +823,7 @@ int32 FPerforceConnection::EditPendingChangelist(const FText& NewDescription, in
 
 		P4Client.SetBreak(NULL);
 
-		return User.ChangelistNumber;
+		return (OutErrorMessages.Num() == 0 ? User.ChangelistNumber : 0);
 	}
 #else
 	return 0;
diff --git a/Engine/Plugins/Developer/PerforceSourceControl/Source/PerforceSourceControl/Private/PerforceSourceControlOperations.cpp b/Engine/Plugins/Developer/PerforceSourceControl/Source/PerforceSourceControl/Private/PerforceSourceControlOperations.cpp
--- a/Engine/Plugins/Developer/PerforceSourceControl/Source/PerforceSourceControl/Private/PerforceSourceControlOperations.cpp
+++ b/Engine/Plugins/Developer/PerforceSourceControl/Source/PerforceSourceControl/Private/PerforceSourceControlOperations.cpp
@@ -497,7 +497,7 @@ bool FPerforceCheckInWorker::Execute(FPerforceSourceControlCommand& InCommand)
 			// Reopen the assets to the default changelist to remove them from the changelist we created above
 			if (ReopenedFiles.Num() > 0)
 			{
-				InCommand.bCommandSuccessful = RunReopenCommand(InCommand, ReopenedFiles, InCommand.Changelist);
+				RunReopenCommand(InCommand, ReopenedFiles, InCommand.Changelist);
 			}
 
 			// Delete the changelist we created above
@@ -1747,6 +1747,8 @@ bool FPerforceGetPendingChangelistsWorker::Execute(FPerforceSourceControlCommand
 			InCommand.bCommandSuccessful &= (InCommand.ResultInfo.ErrorMessages.Num() == 0);
 
 			ParseChangelistsResults(Records, OutChangelistsStates);
+
+			bCleanupCache = InCommand.bCommandSuccessful;
 		}
 
 		// Test whether we should continue processing SCC commands
@@ -1912,6 +1914,28 @@ bool FPerforceGetPendingChangelistsWorker::UpdateStates() const
 		}
 	}
 
+	if (bCleanupCache)
+	{
+		TArray<FPerforceSourceControlChangelist> ChangelistsToRemove;
+		PerforceSourceControl.GetProvider().GetCachedStateByPredicate([this, &ChangelistsToRemove](const FSourceControlChangelistStateRef& InCLState) {
+			TSharedRef<FPerforceSourceControlChangelistState, ESPMode::ThreadSafe> CLState = StaticCastSharedRef<FPerforceSourceControlChangelistState>(InCLState);
+
+			if (Algo::NoneOf(OutChangelistsStates, [&CLState](const FPerforceSourceControlChangelistState& UpdatedCLState) {
+				return CLState->Changelist == UpdatedCLState.Changelist;
+				}))
+			{
+				ChangelistsToRemove.Add(CLState->Changelist);
+			}
+			
+			return false;
+			});
+
+		for (const FPerforceSourceControlChangelist& ChangelistToRemove : ChangelistsToRemove)
+		{
+			PerforceSourceControl.GetProvider().RemoveChangelistFromCache(ChangelistToRemove);
+		}
+	}
+
 	return bUpdated;
 }
 
@@ -2305,27 +2329,87 @@ bool FPerforceShelveWorker::Execute(class FPerforceSourceControlCommand& InComma
 	if (!InCommand.IsCanceled() && ScopedConnection.IsValid())
 	{
 		FPerforceConnection& Connection = ScopedConnection.GetConnection();
-		TArray<FString> Parameters;
-		Parameters.Add(TEXT("-c"));
-		Parameters.Add(InCommand.Changelist.ToString());
-		Parameters.Add(TEXT("-f")); // force
 
-		if (InCommand.Files.Num() > 0)
+		check(InCommand.Operation->GetName() == GetName());
+		TSharedRef<FShelve, ESPMode::ThreadSafe> Operation = StaticCastSharedRef<FShelve>(InCommand.Operation);
+
+		FPerforceSourceControlChangelist Changelist(InCommand.Changelist);
+
+		InCommand.bCommandSuccessful = true;
+
+		// If the command is issued on the default changelist, then we should create a new changelist,
+		// move the files to the new changelist (reopen), then shelve the files
+		if (InCommand.Changelist.IsDefault())
 		{
-			Parameters.Append(InCommand.Files);
+			TArray<FString> FilesToShelve = InCommand.Files;
+
+			// If the command has specified the default changelist but no files, then get all files from the default changelist
+			if (FilesToShelve.Num() == 0 && InCommand.Changelist.IsInitialized())
+			{
+				FPerforceSourceControlModule& PerforceSourceControl = FPerforceSourceControlModule::Get();
+				TSharedRef<FPerforceSourceControlChangelistState, ESPMode::ThreadSafe> DefaultChangelistState = PerforceSourceControl.GetProvider().GetStateInternal(InCommand.Changelist);
+				Algo::Transform(DefaultChangelistState->Files, FilesToShelve, [](const auto& FileState) {
+					return FileState->GetFilename();
+					});
+			}
+
+			int32 NewChangeList = Connection.CreatePendingChangelist(Operation->GetDescription(), TArray<FString>(), FOnIsCancelled::CreateRaw(&InCommand, &FPerforceSourceControlCommand::IsCanceled), InCommand.ResultInfo.ErrorMessages);
+			if (NewChangeList > 0)
+			{
+				Changelist = FPerforceSourceControlChangelist(NewChangeList);
+				InCommand.bCommandSuccessful = RunReopenCommand(InCommand, FilesToShelve, Changelist, &MovedFiles);
+				ChangelistDescription = Operation->GetDescription().ToString();
+			}
+			else
+			{
+				InCommand.bCommandSuccessful = false;
+			}
 		}
 
 		FP4RecordSet Records;
-		InCommand.bCommandSuccessful = Connection.RunCommand(TEXT("shelve"), Parameters, Records, InCommand.ResultInfo.ErrorMessages, FOnIsCancelled::CreateRaw(&InCommand, &FPerforceSourceControlCommand::IsCanceled), InCommand.bConnectionDropped);
-		
+
 		if (InCommand.bCommandSuccessful)
 		{
-			ChangelistToUpdate = InCommand.Changelist;
+			TArray<FString> Parameters;
+			Parameters.Add(TEXT("-c"));
+			Parameters.Add(Changelist.ToString());
+			Parameters.Add(TEXT("-f")); // force
+
+			if (InCommand.Files.Num() > 0)
+			{
+				Parameters.Append(InCommand.Files);
+			}
+
+			InCommand.bCommandSuccessful = Connection.RunCommand(TEXT("shelve"), Parameters, Records, InCommand.ResultInfo.ErrorMessages, FOnIsCancelled::CreateRaw(&InCommand, &FPerforceSourceControlCommand::IsCanceled), InCommand.bConnectionDropped);
+		}
+
+		if (InCommand.bCommandSuccessful)
+		{
+			InChangelistToUpdate = InCommand.Changelist;
+			OutChangelistToUpdate = Changelist;
+
 			ParseShelvedResults(Records, OutResults);
 
 			// Build depot to file mapping
 			GetDepotFileToLocalFileMap(Connection, InCommand, OutResults, OutFileMap);
 		}
+		else
+		{
+			// If we had to create a new changelist, move the files back to the default changelist
+			// and delete the changelist
+			if (Changelist != InCommand.Changelist)
+			{
+				if (MovedFiles.Num() > 0)
+				{
+					RunReopenCommand(InCommand, MovedFiles, InCommand.Changelist);
+				}
+
+				TArray<FString> ChangeParams;
+				ChangeParams.Add(TEXT("-d"));
+				ChangeParams.Add(Changelist.ToString());
+				Connection.RunCommand(TEXT("change"), ChangeParams, Records, InCommand.ResultInfo.ErrorMessages, FOnIsCancelled::CreateRaw(&InCommand, &FPerforceSourceControlCommand::IsCanceled), InCommand.bConnectionDropped);
+			}
+		}
 	}
 
 	return InCommand.bCommandSuccessful;
@@ -2333,7 +2417,37 @@ bool FPerforceShelveWorker::Execute(class FPerforceSourceControlCommand& InComma
 
 bool FPerforceShelveWorker::UpdateStates() const
 {
-	return (OutResults.Num() > 0 && AddShelvedFilesToChangelist(OutResults, OutFileMap, ChangelistToUpdate));
+	FPerforceSourceControlModule& PerforceSourceControl = FPerforceSourceControlModule::Get();
+
+	bool bMovedFiles = false;
+
+	// If we moved files to a new changelist, then we must make sure that the files are properly moved
+	if (InChangelistToUpdate != OutChangelistToUpdate && MovedFiles.Num() > 0)
+	{
+		const FDateTime Now = FDateTime::Now();
+		TSharedRef<FPerforceSourceControlChangelistState, ESPMode::ThreadSafe> SourceChangelistState = PerforceSourceControl.GetProvider().GetStateInternal(InChangelistToUpdate);
+		TSharedRef<FPerforceSourceControlChangelistState, ESPMode::ThreadSafe> DestinationChangelistState = PerforceSourceControl.GetProvider().GetStateInternal(OutChangelistToUpdate);
+
+		DestinationChangelistState->Changelist = OutChangelistToUpdate;
+		DestinationChangelistState->Description = ChangelistDescription;
+		DestinationChangelistState->bHasShelvedFiles = true;
+
+		for (const FString& MovedFile : MovedFiles)
+		{
+			TSharedRef<FPerforceSourceControlState, ESPMode::ThreadSafe> FileState = PerforceSourceControl.GetProvider().GetStateInternal(MovedFile);
+
+			SourceChangelistState->Files.Remove(FileState);
+			DestinationChangelistState->Files.Add(FileState);
+			FileState->Changelist = OutChangelistToUpdate;
+			FileState->TimeStamp = Now;
+		}
+
+		bMovedFiles = true;
+	}
+
+	const bool bAddedShelvedFilesToChangelist = (OutResults.Num() > 0 && AddShelvedFilesToChangelist(OutResults, OutFileMap, OutChangelistToUpdate));
+
+	return bMovedFiles || bAddedShelvedFilesToChangelist;
 }
 
 
diff --git a/Engine/Plugins/Developer/PerforceSourceControl/Source/PerforceSourceControl/Private/PerforceSourceControlOperations.h b/Engine/Plugins/Developer/PerforceSourceControl/Source/PerforceSourceControl/Private/PerforceSourceControlOperations.h
--- a/Engine/Plugins/Developer/PerforceSourceControl/Source/PerforceSourceControl/Private/PerforceSourceControlOperations.h
+++ b/Engine/Plugins/Developer/PerforceSourceControl/Source/PerforceSourceControl/Private/PerforceSourceControlOperations.h
@@ -164,6 +164,10 @@ public:
 	TArray<TArray<FPerforceSourceControlState>> OutCLFilesStates;
 	TArray<TMap<FString, EPerforceState::Type>> OutCLShelvedFilesStates;
 	TArray<TMap<FString, FString>> OutCLShelvedFilesMap;
+
+private:
+	/** Controls whether or not we will remove changelists from the cache after a full update */
+	bool bCleanupCache = false;
 };
 
 class FPerforceCopyWorker : public IPerforceSourceControlWorker
@@ -287,15 +291,22 @@ public:
 	virtual bool Execute(class FPerforceSourceControlCommand& InCommand) override;
 	virtual bool UpdateStates() const override;
 
-protected:	
+protected:
 	/** Map of filenames to perforce state */
 	TMap<FString, EPerforceState::Type> OutResults;
 
 	/** Map depot filenames to local file */
 	TMap<FString, FString> OutFileMap;
 
-	/** Changelist to be updated */
-	FPerforceSourceControlChangelist ChangelistToUpdate;
+	/** Reopened files */
+	TArray<FString> MovedFiles;
+
+	/** Changelist description if needed */
+	FString ChangelistDescription;
+
+	/** Changelist(s) to be updated */
+	FPerforceSourceControlChangelist InChangelistToUpdate;
+	FPerforceSourceControlChangelist OutChangelistToUpdate;
 };
 
 class FPerforceDeleteShelveWorker : public IPerforceSourceControlWorker
diff --git a/Engine/Plugins/Developer/PerforceSourceControl/Source/PerforceSourceControl/Private/PerforceSourceControlProvider.cpp b/Engine/Plugins/Developer/PerforceSourceControl/Source/PerforceSourceControl/Private/PerforceSourceControlProvider.cpp
--- a/Engine/Plugins/Developer/PerforceSourceControl/Source/PerforceSourceControl/Private/PerforceSourceControlProvider.cpp
+++ b/Engine/Plugins/Developer/PerforceSourceControl/Source/PerforceSourceControl/Private/PerforceSourceControlProvider.cpp
@@ -319,6 +319,20 @@ TArray<FSourceControlStateRef> FPerforceSourceControlProvider::GetCachedStateByP
 	return Result;
 }
 
+TArray<FSourceControlChangelistStateRef> FPerforceSourceControlProvider::GetCachedStateByPredicate(TFunctionRef<bool(const FSourceControlChangelistStateRef&)> Predicate) const
+{
+	TArray<FSourceControlChangelistStateRef> Result;
+	for (const auto& CacheItem : ChangelistsStateCache)
+	{
+		FSourceControlChangelistStateRef State = CacheItem.Value;
+		if (Predicate(State))
+		{
+			Result.Add(State);
+		}
+	}
+	return Result;
+}
+
 FDelegateHandle FPerforceSourceControlProvider::RegisterSourceControlStateChanged_Handle( const FSourceControlStateChanged::FDelegate& SourceControlStateChanged )
 {
 	return OnSourceControlStateChanged.Add( SourceControlStateChanged );
diff --git a/Engine/Plugins/Developer/PerforceSourceControl/Source/PerforceSourceControl/Private/PerforceSourceControlProvider.h b/Engine/Plugins/Developer/PerforceSourceControl/Source/PerforceSourceControl/Private/PerforceSourceControlProvider.h
--- a/Engine/Plugins/Developer/PerforceSourceControl/Source/PerforceSourceControl/Private/PerforceSourceControlProvider.h
+++ b/Engine/Plugins/Developer/PerforceSourceControl/Source/PerforceSourceControl/Private/PerforceSourceControlProvider.h
@@ -110,6 +110,9 @@ public:
 	/** Remove a changelist from the state cache */
 	bool RemoveChangelistFromCache(const FPerforceSourceControlChangelist& Changelist);
 
+	/** Returns a list of changelists from the cache based on a given predicate */
+	TArray<FSourceControlChangelistStateRef> GetCachedStateByPredicate(TFunctionRef<bool(const FSourceControlChangelistStateRef&)> Predicate) const;
+
 private:
 
 	/** Helper function used to create a worker for a particular operation */
diff --git a/Engine/Source/Developer/SourceControl/Public/SourceControlOperations.h b/Engine/Source/Developer/SourceControl/Public/SourceControlOperations.h
--- a/Engine/Source/Developer/SourceControl/Public/SourceControlOperations.h
+++ b/Engine/Source/Developer/SourceControl/Public/SourceControlOperations.h
@@ -565,6 +565,20 @@ public:
 	{
 		return LOCTEXT("SourceControl_ShelveOperation", "Shelving files in changelist...");
 	}
+
+	void SetDescription(const FText& InDescription)
+	{
+		Description = InDescription;
+	}
+
+	const FText& GetDescription() const
+	{
+		return Description;
+	}
+
+private:
+	/** Description of the changelist, will be used only to create a new changelist when needed */
+	FText Description;
 };
 
 /**
diff --git a/Engine/Source/Editor/SourceControlWindows/Private/SSourceControlChangelists.cpp b/Engine/Source/Editor/SourceControlWindows/Private/SSourceControlChangelists.cpp
--- a/Engine/Source/Editor/SourceControlWindows/Private/SSourceControlChangelists.cpp
+++ b/Engine/Source/Editor/SourceControlWindows/Private/SSourceControlChangelists.cpp
@@ -473,9 +473,34 @@ bool SSourceControlChangelistsWidget::CanRevert()
 
 void SSourceControlChangelistsWidget::OnShelve()
 {
+	FSourceControlChangelistStatePtr CurrentChangelist = GetChangelistStateFromSelection();
+
+	if (!CurrentChangelist)
+	{
+		return;
+	}
+
+	FText ChangelistDescription = CurrentChangelist->GetDescriptionText();
+
+	if (ChangelistDescription.IsEmptyOrWhitespace())
+	{
+		bool bOk = GetChangelistDescription(
+			nullptr,
+			LOCTEXT("SourceControl.Changelist.NewShelve", "Shelving files..."),
+			LOCTEXT("SourceControl.Changelist.NewShelve.Label", "Enter a description for the changelist holding the shelve:"),
+			ChangelistDescription);
+
+		if (!bOk)
+		{
+			// User cancelled entering a changelist description; abort shelve
+			return;
+		}
+	}
+
 	ISourceControlProvider& SourceControlProvider = ISourceControlModule::Get().GetProvider();
 	auto ShelveOperation = ISourceControlOperation::Create<FShelve>();
-	SourceControlProvider.Execute(ShelveOperation, GetChangelistFromSelection(), GetSelectedFiles());
+	ShelveOperation->SetDescription(ChangelistDescription);
+	SourceControlProvider.Execute(ShelveOperation, CurrentChangelist->GetChangelist(), GetSelectedFiles());
 }
 
 void SSourceControlChangelistsWidget::OnUnshelve()
@@ -900,6 +925,8 @@ public:
 	FText GetChangelistDescriptionText() const
 	{
 		FString DescriptionString = TreeItem->GetDescriptionText().ToString();
+		// Here we'll both remove \r\n (when edited from the dialog) and \n (when we get it from the SCC)
+		DescriptionString.ReplaceInline(TEXT("\r"), TEXT(""));
 		DescriptionString.ReplaceInline(TEXT("\n"), TEXT(" "));
 		DescriptionString.TrimEndInline();
 		return FText::FromString(DescriptionString);
