Fixes temporary CL left being when submitting default changelist.
Changing CL description just before submitting if needed, only if it is not the default CL.
Restore previous changelist description (before tag), in case the submit failed.

#changelist validated

diff --git a/Engine/Plugins/Developer/PerforceSourceControl/Source/PerforceSourceControl/Private/PerforceSourceControlOperations.cpp b/Engine/Plugins/Developer/PerforceSourceControl/Source/PerforceSourceControl/Private/PerforceSourceControlOperations.cpp
--- a/Engine/Plugins/Developer/PerforceSourceControl/Source/PerforceSourceControl/Private/PerforceSourceControlOperations.cpp
+++ b/Engine/Plugins/Developer/PerforceSourceControl/Source/PerforceSourceControl/Private/PerforceSourceControlOperations.cpp
@@ -412,6 +412,7 @@ bool FPerforceCheckInWorker::Execute(FPerforceSourceControlCommand& InCommand)
 	if (!InCommand.IsCanceled() && ScopedConnection.IsValid())
 	{
 		FPerforceConnection& Connection = ScopedConnection.GetConnection();
+		FText CachedDescription;
 
 		check(InCommand.Operation->GetName() == GetName());
 		TSharedRef<FCheckIn, ESPMode::ThreadSafe> Operation = StaticCastSharedRef<FCheckIn>(InCommand.Operation);
@@ -446,6 +447,17 @@ bool FPerforceCheckInWorker::Execute(FPerforceSourceControlCommand& InCommand)
 				InCommand.bCommandSuccessful = false;
 			}
 		}
+		else if (!Operation->GetDescription().IsEmpty())
+		{
+			FPerforceSourceControlModule& PerforceSourceControl = FPerforceSourceControlModule::Get();
+			FSourceControlChangelistStateRef ChangelistState = PerforceSourceControl.GetProvider().GetStateInternal(InCommand.Changelist);
+
+			// Retrieves cached description to restore it in case of submit failure.
+			CachedDescription = ChangelistState->GetDescriptionText();
+
+			int32 UpdatedChangelistNumber = Connection.EditPendingChangelist(Operation->GetDescription(), InCommand.Changelist.ToInt(), FOnIsCancelled::CreateRaw(&InCommand, &FPerforceSourceControlCommand::IsCanceled), InCommand.ResultInfo.ErrorMessages);
+			InCommand.bCommandSuccessful = (UpdatedChangelistNumber != 0) && (UpdatedChangelistNumber == InCommand.Changelist.ToInt());
+		}
 
 		// Only submit if reopen was successful (when starting from the default changelist) or always otherwise
 		if (InCommand.bCommandSuccessful)
@@ -491,22 +503,30 @@ bool FPerforceCheckInWorker::Execute(FPerforceSourceControlCommand& InCommand)
 			}
 		}
 
-		// If the submit failed, clean up the changelist created above
-		if (!InCommand.bCommandSuccessful && InCommand.Changelist.IsDefault())
+		if (!InCommand.bCommandSuccessful)
 		{
-			// Reopen the assets to the default changelist to remove them from the changelist we created above
-			if (ReopenedFiles.Num() > 0)
+			// If the submit failed, clean up the changelist created above
+			if (InCommand.Changelist.IsDefault())
 			{
-				RunReopenCommand(InCommand, ReopenedFiles, InCommand.Changelist);
-			}
+				// Reopen the assets to the default changelist to remove them from the changelist we created above
+				if (ReopenedFiles.Num() > 0)
+				{
+					RunReopenCommand(InCommand, ReopenedFiles, InCommand.Changelist);
+				}
 
-			// Delete the changelist we created above
+				// Delete the changelist we created above
+				{
+					FP4RecordSet Records;
+					TArray<FString> ChangeParams;
+					ChangeParams.Add(TEXT("-d"));
+					ChangeParams.Add(ChangeList.ToString());
+					Connection.RunCommand(TEXT("change"), ChangeParams, Records, InCommand.ResultInfo.ErrorMessages, FOnIsCancelled::CreateRaw(&InCommand, &FPerforceSourceControlCommand::IsCanceled), InCommand.bConnectionDropped);
+				}
+			}
+			// If we modified description and we have the previous one, restore it.
+			else if ((!Operation->GetDescription().IsEmpty()) && (!CachedDescription.IsEmpty()))
 			{
-				FP4RecordSet Records;
-				TArray<FString> ChangeParams;
-				ChangeParams.Add(TEXT("-d"));
-				ChangeParams.Add(ChangeList.ToString());
-				Connection.RunCommand(TEXT("change"), ChangeParams, Records, InCommand.ResultInfo.ErrorMessages, FOnIsCancelled::CreateRaw(&InCommand, &FPerforceSourceControlCommand::IsCanceled), InCommand.bConnectionDropped);
+				Connection.EditPendingChangelist(CachedDescription, InCommand.Changelist.ToInt(), FOnIsCancelled::CreateRaw(&InCommand, &FPerforceSourceControlCommand::IsCanceled), InCommand.ResultInfo.ErrorMessages);
 			}
 		}
 	}
diff --git a/Engine/Source/Editor/SourceControlWindows/Private/SSourceControlChangelists.cpp b/Engine/Source/Editor/SourceControlWindows/Private/SSourceControlChangelists.cpp
--- a/Engine/Source/Editor/SourceControlWindows/Private/SSourceControlChangelists.cpp
+++ b/Engine/Source/Editor/SourceControlWindows/Private/SSourceControlChangelists.cpp
@@ -703,7 +703,7 @@ void SSourceControlChangelistsWidget::OnSubmitChangelist()
 	// Build list of states for the dialog
 	const FText OriginalChangelistDescription = ChangelistState->GetDescriptionText();
 	const bool bAskForChangelistDescription = (OriginalChangelistDescription.IsEmptyOrWhitespace());
-	FText ChangelistDescriptionToSubmit = UpdateChangelistDescriptionToSubmitIfNeeded(bValidationResult, OriginalChangelistDescription, ChangelistState);
+	FText ChangelistDescriptionToSubmit = UpdateChangelistDescriptionToSubmitIfNeeded(bValidationResult, OriginalChangelistDescription);
 
 	TSharedRef<SWindow> NewWindow = SNew(SWindow)
 		.Title(NSLOCTEXT("SourceControl.ConfirmSubmit", "Title", "Confirm changelist submit"))
@@ -733,15 +733,15 @@ void SSourceControlChangelistsWidget::OnSubmitChangelist()
 	if (SourceControlWidget->GetResult() == ESubmitResults::SUBMIT_ACCEPTED)
 	{
 		ISourceControlProvider& SourceControlProvider = ISourceControlModule::Get().GetProvider();
+		FChangeListDescription Description;
 		auto SubmitChangelistOperation = ISourceControlOperation::Create<FCheckIn>();
 		bool bCheckinSuccess = false;
 
-		// Replace description only if there was none
-		if (bAskForChangelistDescription)
-		{
-			FChangeListDescription Description;
-			SourceControlWidget->FillChangeListDescription(Description);
+		SourceControlWidget->FillChangeListDescription(Description);
 
+		// If the description was modified, we add it to the operation to update the changelist
+		if (!OriginalChangelistDescription.EqualTo(Description.Description))
+		{
 			SubmitChangelistOperation->SetDescription(Description.Description);
 		}
 
@@ -763,14 +763,9 @@ void SSourceControlChangelistsWidget::OnSubmitChangelist()
 
 		FSlateNotificationManager::Get().AddNotification(Info);
 	}
-	// If user cancelled or closed the window, we restore the original changelist description
-	else if (HasValidationTag(ChangelistDescriptionToSubmit))
-	{
-		EditChangelistDescription(OriginalChangelistDescription, ChangelistState);
-	}
 }
 
-FText SSourceControlChangelistsWidget::UpdateChangelistDescriptionToSubmitIfNeeded(const bool bInValidationResult, const FText& InOriginalChangelistDescription, const FSourceControlChangelistStatePtr& InChangelistState) const
+FText SSourceControlChangelistsWidget::UpdateChangelistDescriptionToSubmitIfNeeded(const bool bInValidationResult, const FText& InOriginalChangelistDescription) const
 {
 	if (bInValidationResult && USourceControlSettings::IsValidationTagEnabled() && (!HasValidationTag(InOriginalChangelistDescription)))
 	{
@@ -782,8 +777,6 @@ FText SSourceControlChangelistsWidget::UpdateChangelistDescriptionToSubmitIfNeed
 
 		FText ChangelistDescription = FText::FromString(Str);
 
-		EditChangelistDescription(ChangelistDescription, InChangelistState);
-
 		return ChangelistDescription;
 	}
 
diff --git a/Engine/Source/Editor/SourceControlWindows/Private/SSourceControlChangelists.h b/Engine/Source/Editor/SourceControlWindows/Private/SSourceControlChangelists.h
--- a/Engine/Source/Editor/SourceControlWindows/Private/SSourceControlChangelists.h
+++ b/Engine/Source/Editor/SourceControlWindows/Private/SSourceControlChangelists.h
@@ -106,15 +106,14 @@ private:
 	TSharedRef<SWidget> MakeToolBar();
 
 	/**
-	 * Modifies the changelist description if needed, appending validation tag.
+	 * Returns a new changelist description if needed, appending validation tag.
 	 * 
 	 * @param bInValidationResult	The result of the validation step
 	 * @param InOriginalChangelistDescription	Description of the changelist before modification
-	 * @param InChangelistState	Changelist to be modified
 	 * 
 	 * @return The new changelist description
 	 */
-	FText UpdateChangelistDescriptionToSubmitIfNeeded(const bool bInValidationResult, const FText& InOriginalChangelistDescription, const FSourceControlChangelistStatePtr& InChangelistState) const;
+	FText UpdateChangelistDescriptionToSubmitIfNeeded(const bool bInValidationResult, const FText& InOriginalChangelistDescription) const;
 	
 	/** Returns true if the provided changelist description contains a validation tag. */
 	bool HasValidationTag(const FText& InChangelistDescription) const;
