Removes FStatus operation in SourceControl.
Adds Modified State check by Hash for PerforceSourceControlProvider.

#changelist validated

diff --git a/Engine/Plugins/Developer/PerforceSourceControl/Source/PerforceSourceControl/Private/PerforceSourceControlModule.cpp b/Engine/Plugins/Developer/PerforceSourceControl/Source/PerforceSourceControl/Private/PerforceSourceControlModule.cpp
--- a/Engine/Plugins/Developer/PerforceSourceControl/Source/PerforceSourceControl/Private/PerforceSourceControlModule.cpp
+++ b/Engine/Plugins/Developer/PerforceSourceControl/Source/PerforceSourceControl/Private/PerforceSourceControlModule.cpp
@@ -27,7 +27,6 @@ void FPerforceSourceControlModule::StartupModule()
 	PerforceSourceControlProvider.RegisterWorker( "GetWorkspaces", FGetPerforceSourceControlWorker::CreateStatic( &CreateWorker<FPerforceGetWorkspacesWorker> ) );
 	PerforceSourceControlProvider.RegisterWorker( "Copy", FGetPerforceSourceControlWorker::CreateStatic( &CreateWorker<FPerforceCopyWorker> ) );
 	PerforceSourceControlProvider.RegisterWorker( "Resolve", FGetPerforceSourceControlWorker::CreateStatic( &CreateWorker<FPerforceResolveWorker> ) );
-	PerforceSourceControlProvider.RegisterWorker( "Status", FGetPerforceSourceControlWorker::CreateStatic(&CreateWorker<FPerforceStatusWorker>));
 	PerforceSourceControlProvider.RegisterWorker( "ChangeStatus", FGetPerforceSourceControlWorker::CreateStatic(&CreateWorker<FPerforceChangeStatusWorker>));
 	PerforceSourceControlProvider.RegisterWorker( "UpdateChangelistsStatus", FGetPerforceSourceControlWorker::CreateStatic( &CreateWorker<FPerforceGetPendingChangelistsWorker> ) );
 	PerforceSourceControlProvider.RegisterWorker( "NewChangelist", FGetPerforceSourceControlWorker::CreateStatic( &CreateWorker<FPerforceNewChangelistWorker> ) );
diff --git a/Engine/Plugins/Developer/PerforceSourceControl/Source/PerforceSourceControl/Private/PerforceSourceControlOperations.cpp b/Engine/Plugins/Developer/PerforceSourceControl/Source/PerforceSourceControl/Private/PerforceSourceControlOperations.cpp
--- a/Engine/Plugins/Developer/PerforceSourceControl/Source/PerforceSourceControl/Private/PerforceSourceControlOperations.cpp
+++ b/Engine/Plugins/Developer/PerforceSourceControl/Source/PerforceSourceControl/Private/PerforceSourceControlOperations.cpp
@@ -1,6 +1,7 @@
 // Copyright Epic Games, Inc. All Rights Reserved.
 
 #include "PerforceSourceControlOperations.h"
+#include "Async/Async.h"
 #include "PerforceSourceControlPrivate.h"
 #include "HAL/FileManager.h"
 #include "Misc/Paths.h"
@@ -859,10 +860,29 @@ static void ParseBranchModificationResults(const FP4RecordSet& InRecords, const
 			BranchModifications.Add(BranchFile, FBranchModification(Branch, BranchFile, HeadAction, HeadChange, HeadModTime));
 		}
 	}
+}
+
+struct FSourceControlDigest
+{
+	FString		FullPath;
+	FString		DepotDigest;
+	FMD5Hash	LocalDigest;
+};
+
+static FSourceControlDigest ComputeLocalDigest(FString InFilename, FString InDepotDigest)
+{
+	TRACE_CPUPROFILER_EVENT_SCOPE(FPerforceUpdateStatusWorker_ComputeLocalDigestTask);
 
+	FSourceControlDigest Result;
+
+	Result.FullPath = MoveTemp(InFilename);
+	Result.DepotDigest = MoveTemp(InDepotDigest);
+	Result.LocalDigest = FMD5Hash::HashFile(*Result.FullPath, nullptr);
+
+	return Result;
 }
 
-static void ParseUpdateStatusResults(const FP4RecordSet& InRecords, const TArray<FText>& ErrorMessages, TArray<FPerforceSourceControlState>& OutStates, const FString& ContentRoot, TMap<FString, FBranchModification>& BranchModifications)
+static void ParseUpdateStatusResults(const FP4RecordSet& InRecords, const TArray<FText>& ErrorMessages, TArray<FPerforceSourceControlState>& OutStates, TArray<FString>& OutModifiedFiles, const FString& ContentRoot, TMap<FString, FBranchModification>& BranchModifications)
 {
 	// Build up a map of any other branch states
 	for (int32 Index = 0; Index < InRecords.Num(); ++Index)
@@ -927,6 +947,8 @@ static void ParseUpdateStatusResults(const FP4RecordSet& InRecords, const TArray
 
 	}
 
+	TArray<TFuture<FSourceControlDigest>> SourceControlDigests;
+
 	// Iterate over each record found as a result of the command, parsing it for relevant information
 	for (int32 Index = 0; Index < InRecords.Num(); ++Index)
 	{
@@ -941,6 +963,7 @@ static void ParseUpdateStatusResults(const FP4RecordSet& InRecords, const TArray
 		FString HeadAction = ClientRecord(TEXT("headAction"));
 		FString Action = ClientRecord(TEXT("action"));
 		FString HeadType = ClientRecord(TEXT("headType"));
+		FString Digest = ClientRecord(TEXT("digest"));
 		const bool bUnresolved = ClientRecord.Contains(TEXT("unresolved"));
 
 		if (!FileName.Len())
@@ -958,6 +981,12 @@ static void ParseUpdateStatusResults(const FP4RecordSet& InRecords, const TArray
 
 		FString Branch;
 		FString BranchFile;
+
+		if ((Digest.Len() > 0) && (FullPath.Len() > 0))
+		{
+			SourceControlDigests.Add(Async(EAsyncExecution::TaskGraph, [FullPath, Digest]() { return ComputeLocalDigest(FullPath, Digest); }));
+		}
+
 		if (DepotFileName.Split(ContentRoot, &Branch, &BranchFile))
 		{
 			// Sanitize
@@ -1122,6 +1151,23 @@ static void ParseUpdateStatusResults(const FP4RecordSet& InRecords, const TArray
 		}
 	}
 
+	for (const TFuture<FSourceControlDigest>& DigestResult : SourceControlDigests)
+	{
+		FSourceControlDigest Digest = DigestResult.Get();
+
+		if (!Digest.LocalDigest.IsValid())
+		{
+			continue;
+		}
+
+		FString LocalDigest = LexToString(Digest.LocalDigest);
+
+		if (!Digest.DepotDigest.Equals(LocalDigest, ESearchCase::IgnoreCase))
+		{
+			OutModifiedFiles.Add(Digest.FullPath);
+		}
+	}
+
 	// also see if we can glean anything from the error messages
 	for (int32 Index = 0; Index < ErrorMessages.Num(); ++Index)
 	{
@@ -1493,6 +1539,10 @@ bool FPerforceUpdateStatusWorker::Execute(FPerforceSourceControlCommand& InComma
 	if (!InCommand.IsCanceled() && ScopedConnection.IsValid())
 	{
 		FPerforceConnection& Connection = ScopedConnection.GetConnection();
+
+		check(InCommand.Operation->GetName() == GetName());
+		TSharedRef<FUpdateStatus, ESPMode::ThreadSafe> Operation = StaticCastSharedRef<FUpdateStatus>(InCommand.Operation);
+
 		if(InCommand.Files.Num())
 		{
 			// See http://www.perforce.com/perforce/doc.current/manuals/cmdref/p4_fstat.html
@@ -1503,6 +1553,11 @@ bool FPerforceUpdateStatusWorker::Execute(FPerforceSourceControlCommand& InComma
 			// We want to include integration record information:
 			Parameters.Add(TEXT("-Or"));
 
+			if (Operation->ShouldUpdateModifiedState())
+			{
+				Parameters.Add(TEXT("-Olh"));
+			}
+
 			// Get the branches of interest for status updates
 			const FString& ContentRoot = InCommand.ContentRoot;
 			const TArray<FString>& StatusBranches = InCommand.StatusBranchNames;
@@ -1560,7 +1615,7 @@ bool FPerforceUpdateStatusWorker::Execute(FPerforceSourceControlCommand& InComma
 
 			FP4RecordSet Records;
 			InCommand.bCommandSuccessful &= Connection.RunCommand(TEXT("fstat"), Parameters, Records, InCommand.ResultInfo.ErrorMessages, FOnIsCancelled::CreateRaw(&InCommand, &FPerforceSourceControlCommand::IsCanceled), InCommand.bConnectionDropped);
-			ParseUpdateStatusResults(Records, InCommand.ResultInfo.ErrorMessages, OutStates, ContentRoot, BranchModifications);
+			ParseUpdateStatusResults(Records, InCommand.ResultInfo.ErrorMessages, OutStates, OutModifiedFiles,ContentRoot, BranchModifications);
 			RemoveRedundantErrors(InCommand, TEXT(" - no such file(s)."), false);
 			RemoveRedundantErrors(InCommand, TEXT("' is not under client's root '"));
 			RemoveRedundantErrors(InCommand, TEXT(" - protected namespace - access denied"), false);
@@ -1571,8 +1626,6 @@ bool FPerforceUpdateStatusWorker::Execute(FPerforceSourceControlCommand& InComma
 		}
 
 		// update using any special hints passed in via the operation
-		check(InCommand.Operation->GetName() == GetName());
-		TSharedRef<FUpdateStatus, ESPMode::ThreadSafe> Operation = StaticCastSharedRef<FUpdateStatus>(InCommand.Operation);
 
 		bForceQuiet = Operation->ShouldBeQuiet();
 
@@ -2080,89 +2133,6 @@ bool FPerforceResolveWorker::UpdateStates() const
 	return UpdatedFiles.Num() > 0;
 }
 
-FName FPerforceStatusWorker::GetName() const
-{
-	return "Status";
-}
-
-bool FPerforceStatusWorker::Execute(class FPerforceSourceControlCommand& InCommand)
-{
-	FScopedPerforceConnection ScopedConnection(InCommand);
-	
-	if (!InCommand.IsCanceled() && ScopedConnection.IsValid() && (!InCommand.Files.IsEmpty()))
-	{
-		FPerforceConnection& Connection = ScopedConnection.GetConnection();
-
-		TArray<FString> Parameters;
-		Parameters.Append(InCommand.Files);
-
-		RequestedFileStatus = SourceControlHelpers::AbsoluteFilenames(InCommand.Files);
-
-		FP4RecordSet Records;
-		InCommand.bCommandSuccessful = Connection.RunCommand(TEXT("status"), Parameters, Records, InCommand.ResultInfo.ErrorMessages, FOnIsCancelled::CreateRaw(&InCommand, &FPerforceSourceControlCommand::IsCanceled), InCommand.bConnectionDropped);
-		
-		if (InCommand.bCommandSuccessful)
-		{
-			for (const FP4Record& Record : Records)
-			{
-				FString Filename = Record(TEXT("clientFile"));
-				FString Status = Record(TEXT("action"));
-
-				// Here we use the same method as PerforceSourceControlProvider::GetState to be sure to not create duplicates
-				TArray<FString> AbsoluteFilenames = SourceControlHelpers::AbsoluteFilenames({ Filename });
-
-				if (AbsoluteFilenames.IsEmpty() || Status.IsEmpty())
-					continue;
-
-				if (Status.Equals(TEXT("add")))
-				{
-					UpdatedFileStatus.Emplace(AbsoluteFilenames[0], EAction::Add);
-				}
-				else if (Status.Equals(TEXT("edit")))
-				{
-					UpdatedFileStatus.Emplace(AbsoluteFilenames[0], EAction::Edit);
-				}
-				else if (Status.Equals(TEXT("delete")))
-				{
-					UpdatedFileStatus.Emplace(AbsoluteFilenames[0], EAction::Delete);
-				}
-			}
-		}
-	}
-
-	return InCommand.bCommandSuccessful;
-}
-
-bool FPerforceStatusWorker::UpdateStates() const
-{
-	FPerforceSourceControlModule& PerforceSourceControl = FPerforceSourceControlModule::Get();
-	
-	for (const FString& Filename : RequestedFileStatus)
-	{
-		const FFilenameStatusCache::ValueType* StatusPtr = UpdatedFileStatus.Find(Filename);
-		TSharedRef<FPerforceSourceControlState, ESPMode::ThreadSafe> FileState = PerforceSourceControl.GetProvider().GetStateInternal(Filename);
-
-		if ((StatusPtr == nullptr) || (*StatusPtr == EAction::None))
-		{
-			FileState->bModifed = false;
-		}
-		else if (*StatusPtr == EAction::Edit)
-		{
-			FileState->bModifed = true;
-		}
-		else if (*StatusPtr == EAction::Add)
-		{
-			
-		}
-		else if (*StatusPtr == EAction::Delete)
-		{
-
-		}
-	}
-
-	return true;
-}
-
 FName FPerforceChangeStatusWorker::GetName() const
 {
 	return "ChangeStatus";
diff --git a/Engine/Plugins/Developer/PerforceSourceControl/Source/PerforceSourceControl/Private/PerforceSourceControlOperations.h b/Engine/Plugins/Developer/PerforceSourceControl/Source/PerforceSourceControl/Private/PerforceSourceControlOperations.h
--- a/Engine/Plugins/Developer/PerforceSourceControl/Source/PerforceSourceControl/Private/PerforceSourceControlOperations.h
+++ b/Engine/Plugins/Developer/PerforceSourceControl/Source/PerforceSourceControl/Private/PerforceSourceControlOperations.h
@@ -196,31 +196,6 @@ private:
 	TArray< FString > UpdatedFiles;
 };
 
-class FPerforceStatusWorker : public IPerforceSourceControlWorker
-{
-public:
-	enum class EAction
-	{
-		None = 0,
-		Add,
-		Edit,
-		Delete
-	};
-
-	typedef TMap<FString, EAction> FFilenameStatusCache;
-
-public:
-	virtual ~FPerforceStatusWorker() {}
-	// IPerforceSourceControlWorker interface
-	virtual FName GetName() const override;
-	virtual bool Execute(class FPerforceSourceControlCommand& InCommand) override;
-	virtual bool UpdateStates() const override;
-
-private:
-	TArray<FString> RequestedFileStatus;
-	FFilenameStatusCache UpdatedFileStatus;
-};
-
 class FPerforceChangeStatusWorker : public IPerforceSourceControlWorker
 {
 public:
diff --git a/Engine/Source/Developer/SourceControl/Public/SourceControlOperations.h b/Engine/Source/Developer/SourceControl/Public/SourceControlOperations.h
--- a/Engine/Source/Developer/SourceControl/Public/SourceControlOperations.h
+++ b/Engine/Source/Developer/SourceControl/Public/SourceControlOperations.h
@@ -630,22 +630,4 @@ public:
 	}
 };
 
-/**
- * Operation used to get the status of files
- */
-class FStatus : public FSourceControlOperationBase
-{
-public:
-	// ISourceControlOperation interface
-	virtual FName GetName() const override
-	{
-		return "Status";
-	}
-
-	virtual FText GetInProgressString() const override
-	{
-		return LOCTEXT("SourceControl_StatusOperation", "Retrieving files status...");
-	}
-};
-
 #undef LOCTEXT_NAMESPACE
diff --git a/Engine/Source/Developer/UncontrolledChangelists/Private/UncontrolledChangelistState.cpp b/Engine/Source/Developer/UncontrolledChangelists/Private/UncontrolledChangelistState.cpp
--- a/Engine/Source/Developer/UncontrolledChangelists/Private/UncontrolledChangelistState.cpp
+++ b/Engine/Source/Developer/UncontrolledChangelists/Private/UncontrolledChangelistState.cpp
@@ -102,7 +102,9 @@ void FUncontrolledChangelistState::AddFiles(const TArray<FString>& InFilenames,
 
 	if (bCheckStatus)
 	{
-		SourceControlProvider.Execute(ISourceControlOperation::Create<FStatus>(), InFilenames);
+		auto UpdateStatusOperation = ISourceControlOperation::Create<FUpdateStatus>();
+		UpdateStatusOperation->SetUpdateModifiedState(true);
+		SourceControlProvider.Execute(UpdateStatusOperation, InFilenames);
 	}
 
 	bool GetStateSucceeded = SourceControlProvider.GetState(InFilenames, FileStates, EStateCacheUsage::Use) == ECommandResult::Succeeded;
@@ -138,7 +140,10 @@ void FUncontrolledChangelistState::UpdateStatus()
 	}
 
 	ISourceControlProvider& SourceControlProvider = ISourceControlModule::Get().GetProvider();
-	SourceControlProvider.Execute(ISourceControlOperation::Create<FStatus>(), FilesToUpdate, EConcurrency::Asynchronous, FSourceControlOperationComplete::CreateSP(this, &FUncontrolledChangelistState::OnUpdateStatusCompleted));
+	auto UpdateStatusOperation = ISourceControlOperation::Create<FUpdateStatus>();
+	UpdateStatusOperation->SetUpdateModifiedState(true);
+
+	SourceControlProvider.Execute(UpdateStatusOperation, FilesToUpdate, EConcurrency::Asynchronous, FSourceControlOperationComplete::CreateSP(this, &FUncontrolledChangelistState::OnUpdateStatusCompleted));
 }
 
 void FUncontrolledChangelistState::OnUpdateStatusCompleted(const FSourceControlOperationRef& InOperation, ECommandResult::Type InResult)
