Add support for downloading a file from source control directly without actually checking it out (ie p4 print) and add additional support so that the command can be called from a background thread.

- Only implemented for perforce source control at the moment.
- The source control operation FDownload file can be used just like any other command via ::Execute BUT it will also work with the new method ::TryToDownloadFileFromBackgroundThread which unlike ::Execute can be called from a background thread.

#rnx

* PerforceConnection
- Extend ::RunCommand to take an optional FSharedBuffer reference which can be used to return text/binary data downloaded as part of a perforce command
-- Ideally we would replace all of the out parameters with a single struct but that level of refactor would be better suited to its own submit.
- Removed the bool parameter for FP4ClientUser and replaced it with a set of flags instead
- Added a flag to allow FP4ClientUser to store data, either in binary or text format from a perforce command. Although in general this will be the result of a p4 print command.

* PerforceSourceControlCommand
- No longer assert if the command is not on the game thread if the operation indicates that it is safe to run on a background thread.

* PerforceSourceControlModule
- Register the 'DownloadFile' worker

* PerforceSourceControlOperations
- Add FPerforceDownloadFileWorker

* PerforceSourceControlProvider
- Adding implementation for ::TryToDownloadFileFromBackgroundThread.
-- This method is thread safe compared to the normal ::Execute call due to:
-- 1) This method does not accept any callbacks as the most common use of callbacks with ::Execute is to update the UI.
-- 2) This method does not add the command to any of the internal queues nor rely on Tick to process the command.
-- 3) The worker is invoked directly from the calling thread rather than as a background task.
--4) This method does not allow commands that modify the cached states as this would be very difficult to manage without rewriting how all commands are processed to make sure that they are only run one at a time and the results processed in order. NOTE that async commands can be issued from the game thread and can in theory be run in any order on the p4 server depending on how the background task processing works out with no certainty as to the order the results will be processed in. It is unlikely that we will process results out of order but not impossible, but this is a pre-exisiting issue with the current API, extending it to be fully multithreaded will just make it worse.
-- Note that we do not explicity check if the command supports being called from a background thread, as creating the command will check that anyway.
- Calling ::OutputCommandMessages is now thread safe. When called from the game thread it will issue it's output via FMessageLog as before but when called from a background thread we will use the UE_LOG macros instead. Marshalling all output in the same way is something we will need to solve when/if we make the entire API thread safe.

* PerforceControlSettings
- Improve thread safety when calling ::GetConnectionInfo
- If the password is stored in the UI then this is not thread safe but that feature is hidden at the moment due to being unsecure. I am unsure if we want to move the password to be stored in a member variable that can be accessed from any thread and only updated via the UI or not.

* SourceControlOperations
- Add new FDownloadFile operation.
- This can either download the file(s) to FSharedBuffers which can be accessed by calling ::GetFileData OR directly to a directory on disk if the target directory is supplied as part of the constructor.

* ISourceControlOperation
- Add a new base method CanBeCalledFromBackgroundThreads (false by default) which can be overriden by derived classes and return true to indicate that the operation can be run from a background thread. This is intended for use by FDownloadFile with the perforce source control provider in lieu and making the API fully thread safe.

* PerforceSourceControl.uplugin
- Move the perforce plugin to load as early as possible as if Mirage is to use it then we need it to be avaliable before we start loading any potentially virtualized data in packages.

diff --git a/Engine/Plugins/Developer/PerforceSourceControl/PerforceSourceControl.uplugin b/Engine/Plugins/Developer/PerforceSourceControl/PerforceSourceControl.uplugin
--- a/Engine/Plugins/Developer/PerforceSourceControl/PerforceSourceControl.uplugin
+++ b/Engine/Plugins/Developer/PerforceSourceControl/PerforceSourceControl.uplugin
@@ -19,7 +19,7 @@
 		{
 			"Name" : "PerforceSourceControl",
 			"Type" : "UncookedOnly",
-			"LoadingPhase" : "Default"
+			"LoadingPhase" : "EarliestPossible"
 		}
 	]
 }
\ No newline at end of file
diff --git a/Engine/Plugins/Developer/PerforceSourceControl/Source/PerforceSourceControl/Private/PerforceConnection.cpp b/Engine/Plugins/Developer/PerforceSourceControl/Source/PerforceSourceControl/Private/PerforceConnection.cpp
--- a/Engine/Plugins/Developer/PerforceSourceControl/Source/PerforceSourceControl/Private/PerforceConnection.cpp
+++ b/Engine/Plugins/Developer/PerforceSourceControl/Source/PerforceSourceControl/Private/PerforceConnection.cpp
@@ -1,35 +1,49 @@
 // Copyright Epic Games, Inc. All Rights Reserved.
 
 #include "PerforceConnection.h"
-#include "HAL/PlatformProcess.h"
+
 #include "HAL/FileManager.h"
+#include "HAL/PlatformProcess.h"
+#include "HAL/PlatformTime.h"
+#include "ISourceControlModule.h"
+#include "Logging/MessageLog.h"
+#include "Memory/SharedBuffer.h"
 #include "Misc/Paths.h"
 #include "Modules/ModuleManager.h"
-#include "ISourceControlModule.h"
-#include "HAL/PlatformTime.h"
 #include "PerforceSourceControlModule.h"
-#include "Logging/MessageLog.h"
 
 #define LOCTEXT_NAMESPACE "PerforceConnection"
 
 #define FROM_TCHAR(InText, bIsUnicodeServer) (bIsUnicodeServer ? TCHAR_TO_UTF8(InText) : TCHAR_TO_ANSI(InText))
 #define TO_TCHAR(InText, bIsUnicodeServer) (bIsUnicodeServer ? UTF8_TO_TCHAR(InText) : ANSI_TO_TCHAR(InText))
 
+namespace
+{
+	enum class EP4ClientUserFlags
+	{
+		None = 0,
+		UnicodeServer = 1 << 0,
+		CollectData = 1 << 1
+	};
+
+	ENUM_CLASS_FLAGS(EP4ClientUserFlags);
+}
+
 /** Custom ClientUser class for handling results and errors from Perforce commands */
 class FP4ClientUser : public ClientUser
 {
 public:
-
-	// Constructor
-	FP4ClientUser(FP4RecordSet& InRecords, bool bInIsUnicodeServer, TArray<FText>& InOutErrorMessages)
-		:	ClientUser()
-		,	bIsUnicodeServer(bInIsUnicodeServer)
-		,	Records(InRecords)
-		,	OutErrorMessages(InOutErrorMessages)
+	
+	FP4ClientUser(FP4RecordSet& InRecords, EP4ClientUserFlags InFlags, TArray<FText>& InOutErrorMessages)
+		: ClientUser()
+		, Flags(InFlags)
+		, Records(InRecords)
+		, OutErrorMessages(InOutErrorMessages)
 	{
+
 	}
 
-	// Called by P4API when the results from running a command are ready
+	/**  Called by P4API when the results from running a command are ready. */
 	void OutputStat(StrDict* VarList)
 	{
 #if USE_P4_API
@@ -39,25 +53,78 @@ public:
 		// Iterate over each variable and add to records
 		for (int32 Index = 0; VarList->GetVar(Index, Var, Value); Index++)
 		{
-			Record.Add(TO_TCHAR(Var.Text(), bIsUnicodeServer), TO_TCHAR(Value.Text(), bIsUnicodeServer));
+			Record.Add(TO_TCHAR(Var.Text(), IsUnicodeServer()), TO_TCHAR(Value.Text(), IsUnicodeServer()));
+		}
+
+		// Try to reserve the data array if we have one as we just got the filesize info.
+		// NOTE: This code is built on the assumption that the data array will contain the
+		// data sent via p4 print and our code only called p4 print on a single file at the 
+		// time.
+		if (IsCollectingData() && DataBuffer.GetAllocatedSize() == 0)
+		{
+			const FString SizeAsString = Record(TEXT("fileSize"));
+			if (!SizeAsString.IsEmpty())
+			{
+				const int64 Size = FCString::Atoi64(*SizeAsString);
+				DataBuffer.Reserve(Size);
+			}
+
 		}
 
 		Records.Add(Record);
 #endif
 	}
 
-	void HandleError(Error *InError)
+	/** Called by P4API when it output a chunk of text data from a file (commonly via P4 Print) */
+	void OutputText(const char* DataPtr, int DataLength)
+	{
+		if (IsCollectingData())
+		{
+			DataBuffer.Append(DataPtr, DataLength);
+		}
+		else
+		{
+			ClientUser::OutputText(DataPtr, DataLength);
+		}
+	}
+
+	/**  Called by P4API when it output a chunk of binary data from a file (commonly via P4 Print) */
+	void OutputBinary(const char* DataPtr, int DataLength)
+	{
+		if (IsCollectingData())
+		{
+			DataBuffer.Append(DataPtr, DataLength);
+		}
+		else
+		{
+			ClientUser::OutputText(DataPtr, DataLength);
+		}
+	}
+
+	/** Called by P4API when an error message is avaliable. */
+	void HandleError(Error* InError)
 	{
 #if USE_P4_API
 		StrBuf ErrorMessage;
 		InError->Fmt(&ErrorMessage);
-		OutErrorMessages.Add(FText::FromString(FString(TO_TCHAR(ErrorMessage.Text(), bIsUnicodeServer))));
+		OutErrorMessages.Add(FText::FromString(FString(TO_TCHAR(ErrorMessage.Text(), IsUnicodeServer()))));
 #endif
 	}
 
-	bool bIsUnicodeServer;
+	inline bool IsUnicodeServer() const
+	{
+		return EnumHasAnyFlags(Flags, EP4ClientUserFlags::UnicodeServer);
+	}
+
+	inline bool IsCollectingData() const
+	{
+		return EnumHasAnyFlags(Flags, EP4ClientUserFlags::CollectData);
+	}
+
+	EP4ClientUserFlags Flags;
 	FP4RecordSet& Records;
 	TArray<FText>& OutErrorMessages;
+	TArray64<char> DataBuffer;
 };
 
 /** A class used instead of FP4ClientUser for handling changelist create command */
@@ -66,8 +133,8 @@ class FP4CreateChangelistClientUser : public FP4ClientUser
 public:
 
 	// Constructor
-	FP4CreateChangelistClientUser(FP4RecordSet& InRecords, bool bInIsUnicodeServer, TArray<FText>& InOutErrorMessages, const FText& InDescription, const TArray<FString>& InFiles, ClientApi &InP4Client)
-		:	FP4ClientUser(InRecords, bInIsUnicodeServer, InOutErrorMessages)
+	FP4CreateChangelistClientUser(FP4RecordSet& InRecords, EP4ClientUserFlags InFlags, TArray<FText>& InOutErrorMessages, const FText& InDescription, const TArray<FString>& InFiles, ClientApi &InP4Client)
+		:	FP4ClientUser(InRecords, InFlags, InOutErrorMessages)
 		,	Description(InDescription)
 		,	ChangelistNumber(0)
 		,	Files(InFiles)
@@ -79,9 +146,9 @@ public:
 	void OutputInfo(ANSICHAR Level, const ANSICHAR *Data)
 	{
 		const int32 ChangeTextLen = FCString::Strlen(TEXT("Change "));
-		if (FString(TO_TCHAR(Data, bIsUnicodeServer)).StartsWith(TEXT("Change ")))
+		if (FString(TO_TCHAR(Data, IsUnicodeServer())).StartsWith(TEXT("Change ")))
 		{
-			ChangelistNumber = FCString::Atoi(TO_TCHAR(Data + ChangeTextLen, bIsUnicodeServer));
+			ChangelistNumber = FCString::Atoi(TO_TCHAR(Data + ChangeTextLen, IsUnicodeServer()));
 		}
 	}
 
@@ -92,10 +159,10 @@ public:
 		FString OutputDesc;
 		OutputDesc += TEXT("Change:\tnew\n\n");
 		OutputDesc += TEXT("Client:\t");
-		OutputDesc += TO_TCHAR(P4Client.GetClient().Text(), bIsUnicodeServer);
+		OutputDesc += TO_TCHAR(P4Client.GetClient().Text(), IsUnicodeServer());
 		OutputDesc += TEXT("\n\n");
 		OutputDesc += TEXT("User:\t");
-		OutputDesc += TO_TCHAR(P4Client.GetUser().Text(), bIsUnicodeServer);
+		OutputDesc += TO_TCHAR(P4Client.GetUser().Text(), IsUnicodeServer());
 		OutputDesc += TEXT("\n\n");
 		OutputDesc += TEXT("Status:\tnew\n\n");
 		OutputDesc += TEXT("Description:\n");
@@ -119,7 +186,7 @@ public:
 		}
 		OutputDesc += TEXT("\n");
 
-		OutBuffer->Append(FROM_TCHAR(*OutputDesc, bIsUnicodeServer));
+		OutBuffer->Append(FROM_TCHAR(*OutputDesc, IsUnicodeServer()));
 #endif
 	}
 
@@ -134,8 +201,8 @@ class FP4EditChangelistClientUser : public FP4ClientUser
 {
 public:
 	// Constructor
-	FP4EditChangelistClientUser(FP4RecordSet& OutRecords, bool bInIsUnicodeServer, TArray<FText>& InOutErrorMessages, const FText& InDescription, int32 InChangelistNumber, const FP4RecordSet& InRecords, ClientApi& InP4Client)
-		: FP4ClientUser(OutRecords, bInIsUnicodeServer, InOutErrorMessages)
+	FP4EditChangelistClientUser(FP4RecordSet& OutRecords, EP4ClientUserFlags InFlags, TArray<FText>& InOutErrorMessages, const FText& InDescription, int32 InChangelistNumber, const FP4RecordSet& InRecords, ClientApi& InP4Client)
+		: FP4ClientUser(OutRecords, InFlags, InOutErrorMessages)
 		, Description(InDescription)
 		, ChangelistNumber(InChangelistNumber)
 		, P4Client(InP4Client)
@@ -147,9 +214,9 @@ public:
 	void OutputInfo(ANSICHAR Level, const ANSICHAR* Data)
 	{
 		const int32 ChangeTextLen = FCString::Strlen(TEXT("Change "));
-		if (FString(TO_TCHAR(Data, bIsUnicodeServer)).StartsWith(TEXT("Change ")))
+		if (FString(TO_TCHAR(Data, IsUnicodeServer())).StartsWith(TEXT("Change ")))
 		{
-			ChangelistNumber = FCString::Atoi(TO_TCHAR(Data + ChangeTextLen, bIsUnicodeServer));
+			ChangelistNumber = FCString::Atoi(TO_TCHAR(Data + ChangeTextLen, IsUnicodeServer()));
 		}
 	}
 
@@ -213,7 +280,7 @@ public:
 			}
 		}
 
-		OutBuffer->Append(FROM_TCHAR(*OutputDesc, bIsUnicodeServer));
+		OutBuffer->Append(FROM_TCHAR(*OutputDesc, IsUnicodeServer()));
 #endif
 	}
 
@@ -229,8 +296,8 @@ class FP4LoginClientUser : public FP4ClientUser
 public:
 
 	// Constructor
-	FP4LoginClientUser(const FString& InPassword, FP4RecordSet& InRecords, bool bInIsUnicodeServer, TArray<FText>& InOutErrorMessages)
-		:	FP4ClientUser(InRecords, bInIsUnicodeServer, InOutErrorMessages)
+	FP4LoginClientUser(const FString& InPassword, FP4RecordSet& InRecords, EP4ClientUserFlags InFlags, TArray<FText>& InOutErrorMessages)
+		:	FP4ClientUser(InRecords, InFlags, InOutErrorMessages)
 		,	Password(InPassword)
 	{
 	}
@@ -238,7 +305,7 @@ public:
 	void Prompt( const StrPtr& InMessage, StrBuf& OutPrompt, int NoEcho, Error* InError )
 	{
 #if USE_P4_API
-		OutPrompt.Set(FROM_TCHAR(*Password, bIsUnicodeServer));
+		OutPrompt.Set(FROM_TCHAR(*Password, IsUnicodeServer()));
 #endif
 	}
 
@@ -277,7 +344,8 @@ public:
 static bool TestConnection(ClientApi& P4Client, const FString& ClientSpecName, bool bIsUnicodeServer, TArray<FText>& OutErrorMessages)
 {
 	FP4RecordSet Records;
-	FP4ClientUser User(Records, bIsUnicodeServer, OutErrorMessages);
+	EP4ClientUserFlags Flags = bIsUnicodeServer ? EP4ClientUserFlags::UnicodeServer : EP4ClientUserFlags::None;
+	FP4ClientUser User(Records, Flags, OutErrorMessages);
 
 	UTF8CHAR* ClientSpecUTF8Name = nullptr;
 	if(bIsUnicodeServer)
@@ -319,7 +387,7 @@ static bool CheckUnicodeStatus(ClientApi& P4Client, bool& bIsUnicodeServer, TArr
 {
 #if USE_P4_API
 	FP4RecordSet Records;
-	FP4ClientUser User(Records, false, OutErrorMessages);
+	FP4ClientUser User(Records, EP4ClientUserFlags::None, OutErrorMessages);
 
 	P4Client.Run("info", &User);
 
@@ -392,7 +460,7 @@ bool FPerforceConnection::Login(const FPerforceConnectionInfo& InConnectionInfo)
 	TArray<FText> ErrorMessages;
 #if USE_P4_API
 	FP4RecordSet Records;
-	FP4LoginClientUser User(InConnectionInfo.Password, Records, false, ErrorMessages);
+	FP4LoginClientUser User(InConnectionInfo.Password, Records, EP4ClientUserFlags::None, ErrorMessages);
 
 	const char *ArgV[] = { "-a" };
 	P4Client.SetArgv(1, const_cast<char*const*>(ArgV));
@@ -684,7 +752,7 @@ void FPerforceConnection::Disconnect()
 #endif
 }
 
-bool FPerforceConnection::RunCommand(const FString& InCommand, const TArray<FString>& InParameters, FP4RecordSet& OutRecordSet, TArray<FText>& OutErrorMessage, FOnIsCancelled InIsCancelled, bool& OutConnectionDropped, const bool bInStandardDebugOutput, const bool bInAllowRetry)
+bool FPerforceConnection::RunCommand(const FString& InCommand, const TArray<FString>& InParameters, FP4RecordSet& OutRecordSet, TOptional<FSharedBuffer>& OutData, TArray<FText>& OutErrorMessage, FOnIsCancelled InIsCancelled, bool& OutConnectionDropped, const bool bInStandardDebugOutput, const bool bInAllowRetry)
 {
 #if USE_P4_API
 	if (!bEstablishedConnection)
@@ -731,13 +799,23 @@ bool FPerforceConnection::RunCommand(const FString& InCommand, const TArray<FStr
 	P4Client.SetBreak(&KeepAlive);
 
 	OutRecordSet.Reset();
-	FP4ClientUser User(OutRecordSet, bIsUnicode, OutErrorMessage);
+
+	EP4ClientUserFlags Flags = EP4ClientUserFlags::None;
+	Flags |= bIsUnicode ? EP4ClientUserFlags::UnicodeServer : EP4ClientUserFlags::None;
+	Flags |= OutData ? EP4ClientUserFlags::CollectData : EP4ClientUserFlags::None;
+	
+	FP4ClientUser User(OutRecordSet, Flags, OutErrorMessage);
 	P4Client.Run(FROM_TCHAR(*InCommand, bIsUnicode), &User);
 	if ( P4Client.Dropped() )
 	{
 		OutConnectionDropped = true;
 	}
 
+	if (OutData)
+	{
+		OutData = FSharedBuffer::Clone(User.DataBuffer.GetData(), User.DataBuffer.Num());
+	}	
+
 	P4Client.SetBreak(NULL);
 
 	// Free arguments
@@ -767,6 +845,7 @@ int32 FPerforceConnection::CreatePendingChangelist(const FText &Description, con
 #if USE_P4_API
 	TArray<FString> Params;
 	FP4RecordSet Records;
+	EP4ClientUserFlags Flags = bIsUnicode ? EP4ClientUserFlags::UnicodeServer : EP4ClientUserFlags::None;
 
 	const char *ArgV[] = { "-i" };
 	P4Client.SetArgv(1, const_cast<char*const*>(ArgV));
@@ -774,7 +853,7 @@ int32 FPerforceConnection::CreatePendingChangelist(const FText &Description, con
 	FP4KeepAlive KeepAlive(InIsCancelled);
 	P4Client.SetBreak(&KeepAlive);
 
-	FP4CreateChangelistClientUser User(Records, bIsUnicode, OutErrorMessages, Description, Files, P4Client);
+	FP4CreateChangelistClientUser User(Records, Flags, OutErrorMessages, Description, Files, P4Client);
 	P4Client.Run("change", &User);
 
 	P4Client.SetBreak(NULL);
@@ -793,6 +872,7 @@ int32 FPerforceConnection::EditPendingChangelist(const FText& NewDescription, in
 	// Get changelist current specification
 	{
 		TArray<FString> Params;
+		TOptional<FSharedBuffer> CommandData;
 		bool bConnectionDropped = false;
 		const bool bStandardDebugOutput = false;
 		const bool bAllowRetry = true;
@@ -801,7 +881,7 @@ int32 FPerforceConnection::EditPendingChangelist(const FText& NewDescription, in
 		// TODO : make this work also for default changelist, but should really be a Create
 		Params.Add(FString::Printf(TEXT("%d"), ChangelistNumber));
 
-		if (!RunCommand(TEXT("change"), Params, PreviousRecords, OutErrorMessages, InIsCancelled, bConnectionDropped, bStandardDebugOutput, bAllowRetry))
+		if (!RunCommand(TEXT("change"), Params, PreviousRecords, CommandData, OutErrorMessages, InIsCancelled, bConnectionDropped, bStandardDebugOutput, bAllowRetry))
 		{
 			return 0;
 		}
@@ -811,6 +891,7 @@ int32 FPerforceConnection::EditPendingChangelist(const FText& NewDescription, in
 	{
 		TArray<FString> Params;
 		FP4RecordSet Records;
+		EP4ClientUserFlags Flags = bIsUnicode ? EP4ClientUserFlags::UnicodeServer : EP4ClientUserFlags::None;
 
 		const char* ArgV[] = { "-i" };
 		P4Client.SetArgv(1, const_cast<char* const*>(ArgV));
@@ -818,7 +899,7 @@ int32 FPerforceConnection::EditPendingChangelist(const FText& NewDescription, in
 		FP4KeepAlive KeepAlive(InIsCancelled);
 		P4Client.SetBreak(&KeepAlive);
 
-		FP4EditChangelistClientUser User(Records, bIsUnicode, OutErrorMessages, NewDescription, ChangelistNumber, PreviousRecords, P4Client);
+		FP4EditChangelistClientUser User(Records, Flags, OutErrorMessages, NewDescription, ChangelistNumber, PreviousRecords, P4Client);
 		P4Client.Run("change", &User);
 
 		P4Client.SetBreak(NULL);
@@ -879,13 +960,14 @@ void FPerforceConnection::EstablishConnection(const FPerforceConnectionInfo& InC
 		TArray<FString> Params;
 		TArray<FText> ErrorMessages;
 		FP4RecordSet Records;
+		TOptional<FSharedBuffer> CommandData;
 		bool bConnectionDropped = false;
 		const bool bStandardDebugOutput = false;
 		const bool bAllowRetry = true;
 
 		UE_LOG(LogSourceControl, Verbose, TEXT(" ... checking unicode status" ));
 
-		if (RunCommand(TEXT("info"), Params, Records, ErrorMessages, FOnIsCancelled(), bConnectionDropped, bStandardDebugOutput, bAllowRetry))
+		if (RunCommand(TEXT("info"), Params, Records, CommandData, ErrorMessages, FOnIsCancelled(), bConnectionDropped, bStandardDebugOutput, bAllowRetry))
 		{
 			// Get character encoding
 			bIsUnicode = Records[0].Find(TEXT("unicode")) != NULL;
@@ -918,7 +1000,7 @@ void FPerforceConnection::EstablishConnection(const FPerforceConnectionInfo& InC
 			// Gather the client root
 			UE_LOG(LogSourceControl, Verbose, TEXT(" ... getting info" ));
 			bConnectionDropped = false;
-			if (RunCommand(TEXT("info"), Params, Records, ErrorMessages, FOnIsCancelled(), bConnectionDropped, bStandardDebugOutput, bAllowRetry))
+			if (RunCommand(TEXT("info"), Params, Records, CommandData, ErrorMessages, FOnIsCancelled(), bConnectionDropped, bStandardDebugOutput, bAllowRetry))
 			{
 				UE_LOG(LogSourceControl, Verbose, TEXT(" ... getting clientroot" ));
 				ClientRoot = Records[0](TEXT("clientRoot"));
diff --git a/Engine/Plugins/Developer/PerforceSourceControl/Source/PerforceSourceControl/Private/PerforceConnection.h b/Engine/Plugins/Developer/PerforceSourceControl/Source/PerforceSourceControl/Private/PerforceConnection.h
--- a/Engine/Plugins/Developer/PerforceSourceControl/Source/PerforceSourceControl/Private/PerforceConnection.h
+++ b/Engine/Plugins/Developer/PerforceSourceControl/Source/PerforceSourceControl/Private/PerforceConnection.h
@@ -5,6 +5,7 @@
 #include "CoreMinimal.h"
 #include "PerforceSourceControlPrivate.h"
 #include "PerforceSourceControlCommand.h"
+#include "Memory/SharedBuffer.h"
 
 /** A map containing result of running Perforce command */
 class FP4Record : public TMap<FString, FString >
@@ -68,13 +69,22 @@ public:
 	{
 		const bool bStandardDebugOutput=true;
 		const bool bAllowRetry=true;
-		return RunCommand(InCommand, InParameters, OutRecordSet, OutErrorMessage, InIsCancelled, OutConnectionDropped, bStandardDebugOutput, bAllowRetry);
+		TOptional<FSharedBuffer> UnsetDataBuffer;
+
+		return RunCommand(InCommand, InParameters, OutRecordSet, UnsetDataBuffer, OutErrorMessage, InIsCancelled, OutConnectionDropped, bStandardDebugOutput, bAllowRetry);
+	}
+
+	bool RunCommand(const FString& InCommand, const TArray<FString>& InParameters, FP4RecordSet& OutRecordSet, TOptional<FSharedBuffer>& OutData, TArray<FText>& OutErrorMessage, FOnIsCancelled InIsCancelled, bool& OutConnectionDropped)
+	{
+		const bool bStandardDebugOutput = true;
+		const bool bAllowRetry = true;
+		return RunCommand(InCommand, InParameters, OutRecordSet, OutData, OutErrorMessage, InIsCancelled, OutConnectionDropped, bStandardDebugOutput, bAllowRetry);
 	}
 
 	/**
 	 * Runs internal perforce command, catches exceptions, returns results
 	 */
-	bool RunCommand(const FString& InCommand, const TArray<FString>& InParameters, FP4RecordSet& OutRecordSet, TArray<FText>& OutErrorMessage, FOnIsCancelled InIsCancelled, bool& OutConnectionDropped, const bool bInStandardDebugOutput, const bool bInAllowRetry);
+	bool RunCommand(const FString& InCommand, const TArray<FString>& InParameters, FP4RecordSet& OutRecordSet, TOptional<FSharedBuffer>& OutData, TArray<FText>& OutErrorMessage, FOnIsCancelled InIsCancelled, bool& OutConnectionDropped, const bool bInStandardDebugOutput, const bool bInAllowRetry);
 
 	/**
 	 * Creates a changelist with the specified description
diff --git a/Engine/Plugins/Developer/PerforceSourceControl/Source/PerforceSourceControl/Private/PerforceSourceControlCommand.cpp b/Engine/Plugins/Developer/PerforceSourceControl/Source/PerforceSourceControl/Private/PerforceSourceControlCommand.cpp
--- a/Engine/Plugins/Developer/PerforceSourceControl/Source/PerforceSourceControl/Private/PerforceSourceControlCommand.cpp
+++ b/Engine/Plugins/Developer/PerforceSourceControl/Source/PerforceSourceControl/Private/PerforceSourceControlCommand.cpp
@@ -18,8 +18,8 @@ FPerforceSourceControlCommand::FPerforceSourceControlCommand(const TSharedRef<cl
 	, bAutoDelete(true)
 	, Concurrency(EConcurrency::Synchronous)
 {
-	// grab the providers settings here, so we don't access them once the worker thread is launched
-	check(IsInGameThread());
+	// Grab the providers settings here, so we don't access them once the worker thread is launched
+	check(IsInGameThread() || InOperation->CanBeCalledFromBackgroundThreads());
 	FPerforceSourceControlModule& PerforceSourceControl = FModuleManager::LoadModuleChecked<FPerforceSourceControlModule>( "PerforceSourceControl" );
 	ConnectionInfo = PerforceSourceControl.AccessSettings().GetConnectionInfo();
 }
diff --git a/Engine/Plugins/Developer/PerforceSourceControl/Source/PerforceSourceControl/Private/PerforceSourceControlModule.cpp b/Engine/Plugins/Developer/PerforceSourceControl/Source/PerforceSourceControl/Private/PerforceSourceControlModule.cpp
--- a/Engine/Plugins/Developer/PerforceSourceControl/Source/PerforceSourceControl/Private/PerforceSourceControlModule.cpp
+++ b/Engine/Plugins/Developer/PerforceSourceControl/Source/PerforceSourceControl/Private/PerforceSourceControlModule.cpp
@@ -37,7 +37,7 @@ void FPerforceSourceControlModule::StartupModule()
 	PerforceSourceControlProvider.RegisterWorker( "Shelve", FGetPerforceSourceControlWorker::CreateStatic( &CreateWorker<FPerforceShelveWorker>) );
 	PerforceSourceControlProvider.RegisterWorker( "Unshelve", FGetPerforceSourceControlWorker::CreateStatic( &CreateWorker<FPerforceUnshelveWorker>) );
 	PerforceSourceControlProvider.RegisterWorker( "DeleteShelved", FGetPerforceSourceControlWorker::CreateStatic( &CreateWorker<FPerforceDeleteShelveWorker>) );
-
+	PerforceSourceControlProvider.RegisterWorker( "DownloadFile", FGetPerforceSourceControlWorker::CreateStatic( &CreateWorker<FPerforceDownloadFileWorker>) );
 
 	// load our settings
 	PerforceSourceControlSettings.LoadSettings();
diff --git a/Engine/Plugins/Developer/PerforceSourceControl/Source/PerforceSourceControl/Private/PerforceSourceControlOperations.cpp b/Engine/Plugins/Developer/PerforceSourceControl/Source/PerforceSourceControl/Private/PerforceSourceControlOperations.cpp
--- a/Engine/Plugins/Developer/PerforceSourceControl/Source/PerforceSourceControl/Private/PerforceSourceControlOperations.cpp
+++ b/Engine/Plugins/Developer/PerforceSourceControl/Source/PerforceSourceControl/Private/PerforceSourceControlOperations.cpp
@@ -2721,4 +2721,92 @@ bool FPerforceUnshelveWorker::UpdateStates() const
 	}
 }
 
+FName FPerforceDownloadFileWorker::GetName() const
+{
+	return "DownloadFile";
+}
+
+bool FPerforceDownloadFileWorker::Execute(FPerforceSourceControlCommand& InCommand)
+{
+	FScopedPerforceConnection ScopedConnection(InCommand);
+
+	if (!InCommand.IsCanceled() && ScopedConnection.IsValid())
+	{
+		FPerforceConnection& Connection = ScopedConnection.GetConnection();
+		
+		// TODO: The cast doesn't seem to have any type safety
+		TSharedRef<FDownloadFile, ESPMode::ThreadSafe> Operation = StaticCastSharedRef<FDownloadFile>(InCommand.Operation);
+
+		TArray<FString> Parameters;
+		FP4RecordSet Records;
+
+		const FString TargetDirectory = Operation->GetTargetDirectory();
+
+		if (TargetDirectory.IsEmpty())
+		{
+			// Download and store the files as blobs in memory that the caller can access as they wish
+			for (const FString& TargetFilePath : InCommand.Files)
+			{
+				Parameters.Add(TEXT("-q")); // Do not print the header, we only want the actual file contents
+				Parameters.Add(TargetFilePath);
+
+				TOptional<FSharedBuffer> FileData = FSharedBuffer();
+				InCommand.bCommandSuccessful = Connection.RunCommand(TEXT("print"), Parameters, Records, FileData, InCommand.ResultInfo.ErrorMessages, FOnIsCancelled::CreateRaw(&InCommand, &FPerforceSourceControlCommand::IsCanceled), InCommand.bConnectionDropped);
+
+				if (InCommand.bCommandSuccessful)
+				{
+					Operation->AddFileData(TargetFilePath, FileData.GetValue());
+				}
+
+				// Keep the already allocated memory for future files to prevent reallocation
+				Parameters.Reset();
+				Records.Reset();
+			}
+		}
+		else
+		{
+			// Downloading the files directly to a target directory
+			for (const FString& TargetFilePath : InCommand.Files)
+			{
+				// First we will stream to a temp file before moving it to the final location if the download was successful.
+				// This will prevent us from overwriting files with partial data in case of an error.
+
+				Parameters.Add(TEXT("-o"));
+				const FString BaseFilename = FPaths::GetBaseFilename(TargetFilePath);
+				const FString TempFilePath = FPaths::CreateTempFilename(*FPaths::ProjectSavedDir(), *BaseFilename.Left(32));
+				Parameters.Add(TempFilePath);
+
+				Parameters.Add(TEXT("-q")); // Do not print the header, we only want the actual file contents
+				Parameters.Add(TargetFilePath);
+
+				InCommand.bCommandSuccessful = Connection.RunCommand(TEXT("print"), Parameters, Records, InCommand.ResultInfo.ErrorMessages, FOnIsCancelled::CreateRaw(&InCommand, &FPerforceSourceControlCommand::IsCanceled), InCommand.bConnectionDropped);
+
+				if (InCommand.bCommandSuccessful)
+				{
+					const FString FinalTargetPath = TargetDirectory / FPaths::GetCleanFilename(TargetFilePath);
+					if (IFileManager::Get().Move(*FinalTargetPath, *TempFilePath) == false)
+					{
+						// Give error 
+						const FString ErrorMsg = FString::Printf(TEXT("Failed to move '%s' to '%s'"), *TempFilePath, *FinalTargetPath);
+						InCommand.ResultInfo.ErrorMessages.Add(FText::FromString(ErrorMsg));
+						InCommand.bCommandSuccessful = false;
+					}
+				}
+
+				// Keep the already allocated memory for future files to prevent reallocation
+				Parameters.Reset();
+				Records.Reset();
+			}
+		}
+	}
+	return InCommand.bCommandSuccessful;
+}
+
+bool FPerforceDownloadFileWorker::UpdateStates() const
+{
+	// Downloading a file from the server will never affect the cached
+	// file states.
+	return false;
+}
+
 #undef LOCTEXT_NAMESPACE
diff --git a/Engine/Plugins/Developer/PerforceSourceControl/Source/PerforceSourceControl/Private/PerforceSourceControlOperations.h b/Engine/Plugins/Developer/PerforceSourceControl/Source/PerforceSourceControl/Private/PerforceSourceControlOperations.h
--- a/Engine/Plugins/Developer/PerforceSourceControl/Source/PerforceSourceControl/Private/PerforceSourceControlOperations.h
+++ b/Engine/Plugins/Developer/PerforceSourceControl/Source/PerforceSourceControl/Private/PerforceSourceControlOperations.h
@@ -345,3 +345,14 @@ protected:
 	/** List of files states after update */
 	TArray<FPerforceSourceControlState> ChangelistFilesStates;
 };
+
+class FPerforceDownloadFileWorker : public IPerforceSourceControlWorker
+{
+public:
+	virtual ~FPerforceDownloadFileWorker() = default;
+
+	// IPerforceSourceControlWorker interface
+	virtual FName GetName() const override;
+	virtual bool Execute(class FPerforceSourceControlCommand& InCommand) override;
+	virtual bool UpdateStates() const override;
+};
diff --git a/Engine/Plugins/Developer/PerforceSourceControl/Source/PerforceSourceControl/Private/PerforceSourceControlProvider.cpp b/Engine/Plugins/Developer/PerforceSourceControl/Source/PerforceSourceControl/Private/PerforceSourceControlProvider.cpp
--- a/Engine/Plugins/Developer/PerforceSourceControl/Source/PerforceSourceControl/Private/PerforceSourceControlProvider.cpp
+++ b/Engine/Plugins/Developer/PerforceSourceControl/Source/PerforceSourceControl/Private/PerforceSourceControlProvider.cpp
@@ -447,16 +447,31 @@ bool FPerforceSourceControlProvider::UsesCheckout() const
 
 void FPerforceSourceControlProvider::OutputCommandMessages(const FPerforceSourceControlCommand& InCommand) const
 {
-	FMessageLog SourceControlLog("SourceControl");
-
-	for (int32 ErrorIndex = 0; ErrorIndex < InCommand.ResultInfo.ErrorMessages.Num(); ++ErrorIndex)
+	if (IsInGameThread()) // On the game thread we can use FMessageLog
 	{
-		SourceControlLog.Error(FText::Format(LOCTEXT("OutputCommandMessagesFormatError", "CommandMessage Command: {0}, Error: {1}"), FText::FromName(InCommand.Operation->GetName()), InCommand.ResultInfo.ErrorMessages[ErrorIndex]));
-	}
+		FMessageLog SourceControlLog("SourceControl");
+
+		for (int32 ErrorIndex = 0; ErrorIndex < InCommand.ResultInfo.ErrorMessages.Num(); ++ErrorIndex)
+		{
+			SourceControlLog.Error(FText::Format(LOCTEXT("OutputCommandMessagesFormatError", "CommandMessage Command: {0}, Error: {1}"), FText::FromName(InCommand.Operation->GetName()), InCommand.ResultInfo.ErrorMessages[ErrorIndex]));
+		}
 
-	for (int32 InfoIndex = 0; InfoIndex < InCommand.ResultInfo.InfoMessages.Num(); ++InfoIndex)
+		for (int32 InfoIndex = 0; InfoIndex < InCommand.ResultInfo.InfoMessages.Num(); ++InfoIndex)
+		{
+			SourceControlLog.Info(FText::Format(LOCTEXT("OutputCommandMessagesFormatInfo", "CommandMessage Command: {0}, Info: {1}"), FText::FromName(InCommand.Operation->GetName()), InCommand.ResultInfo.InfoMessages[InfoIndex]));
+		}
+	}
+	else // On background threads we must log directly as FMessageLog internals cannot be assumed to be thread safe
 	{
-		SourceControlLog.Info(FText::Format(LOCTEXT("OutputCommandMessagesFormatInfo", "CommandMessage Command: {0}, Info: {1}"), FText::FromName(InCommand.Operation->GetName()), InCommand.ResultInfo.InfoMessages[InfoIndex]));
+		for (const FText& Error : InCommand.ResultInfo.ErrorMessages)
+		{
+			UE_LOG(LogSourceControl, Error, TEXT("Command: %s, Error: %s"), *InCommand.Operation->GetName().ToString(), *Error.ToString());
+		}
+
+		for (const FText& Info : InCommand.ResultInfo.InfoMessages)
+		{
+			UE_LOG(LogSourceControl, Log, TEXT("Command: %s, Info: %s"), *InCommand.Operation->GetName().ToString(), *Info.ToString());
+		}
 	}
 }
 
@@ -585,6 +600,51 @@ TArray<FSourceControlChangelistRef> FPerforceSourceControlProvider::GetChangelis
 	return Changelists;
 }
 
+bool FPerforceSourceControlProvider::TryToDownloadFileFromBackgroundThread(const TSharedRef<class FDownloadFile>& InOperation, const TArray<FString>& InFiles)
+{
+	if (!IsEnabled())
+	{
+		return false;
+	}
+
+	TSharedPtr<IPerforceSourceControlWorker, ESPMode::ThreadSafe> Worker = CreateWorker(InOperation->GetName());
+	if (!Worker.IsValid())
+	{
+		// This operation is unsupported by this source control provider
+		FFormatNamedArguments Arguments;
+		Arguments.Add(TEXT("OperationName"), FText::FromName(InOperation->GetName()));
+		Arguments.Add(TEXT("ProviderName"), FText::FromName(GetName()));
+		FText Message = FText::Format(LOCTEXT("UnsupportedOperation", "Operation '{OperationName}' not supported by source control provider '{ProviderName}'"), Arguments);
+
+		InOperation->AddErrorMessge(Message);
+
+		return false;
+	}
+
+	// Note that this method can safely be called from any thread works because we know that
+	// a) We are not executing any delegates.
+	// b) That the FDownloadFile operation does not change the stage of any files in 
+	// source control and so will not affect any cached states.
+	// c) We are not invoking any globals that might touch the slate UI such as FMessageLog
+
+	FPerforceSourceControlCommand Command(InOperation, Worker.ToSharedRef());
+	Command.bAutoDelete = false;
+	Command.Files = SourceControlHelpers::AbsoluteFilenames(InFiles);
+	Command.StatusBranchNames = StatusBranchNames;
+	Command.ContentRoot = ContentRoot;
+
+	// DoWork will use a shared connection, we need to call DoThreadedWork to make sure that 
+	// we use our own connection for this background thread.
+	Command.DoThreadedWork();
+
+	// Sanity check to make sure we are not running a command that modifies the cached states from a background thread
+	check(Command.Worker->UpdateStates() == false);
+
+	OutputCommandMessages(Command);
+	
+	return Command.bCommandSuccessful;
+}
+
 #if SOURCE_CONTROL_WITH_SLATE
 TSharedRef<class SWidget> FPerforceSourceControlProvider::MakeSettingsWidget() const
 {
diff --git a/Engine/Plugins/Developer/PerforceSourceControl/Source/PerforceSourceControl/Private/PerforceSourceControlProvider.h b/Engine/Plugins/Developer/PerforceSourceControl/Source/PerforceSourceControl/Private/PerforceSourceControlProvider.h
--- a/Engine/Plugins/Developer/PerforceSourceControl/Source/PerforceSourceControl/Private/PerforceSourceControlProvider.h
+++ b/Engine/Plugins/Developer/PerforceSourceControl/Source/PerforceSourceControl/Private/PerforceSourceControlProvider.h
@@ -49,6 +49,9 @@ public:
 	virtual void Tick() override;
 	virtual TArray< TSharedRef<class ISourceControlLabel> > GetLabels( const FString& InMatchingSpec ) const override;
 	virtual TArray<FSourceControlChangelistRef> GetChangelists( EStateCacheUsage::Type InStateCacheUsage ) override;
+
+	virtual bool TryToDownloadFileFromBackgroundThread(const TSharedRef<class FDownloadFile>& InOperation, const TArray<FString>& InFiles) override;
+
 #if SOURCE_CONTROL_WITH_SLATE
 	virtual TSharedRef<class SWidget> MakeSettingsWidget() const override;
 #endif
diff --git a/Engine/Plugins/Developer/PerforceSourceControl/Source/PerforceSourceControl/Private/PerforceSourceControlSettings.cpp b/Engine/Plugins/Developer/PerforceSourceControl/Source/PerforceSourceControl/Private/PerforceSourceControlSettings.cpp
--- a/Engine/Plugins/Developer/PerforceSourceControl/Source/PerforceSourceControl/Private/PerforceSourceControlSettings.cpp
+++ b/Engine/Plugins/Developer/PerforceSourceControl/Source/PerforceSourceControl/Private/PerforceSourceControlSettings.cpp
@@ -104,16 +104,20 @@ void FPerforceSourceControlSettings::SaveSettings() const
 
 FPerforceConnectionInfo FPerforceSourceControlSettings::GetConnectionInfo() const
 {
-	check(IsInGameThread());
-	FPerforceConnectionInfo OutConnectionInfo = ConnectionInfo;
+	FPerforceConnectionInfo OutConnectionInfo;
+	{
+		FScopeLock ScopeLock(&CriticalSection);
+		OutConnectionInfo = ConnectionInfo;
+	}
 	
-	// password needs to be gotten straight from the input UI, its not stored anywhere else
-	if(SPerforceSourceControlSettings::GetPassword().Len() > 0)
+	// The password needs to be gotten straight from the input UI, its not stored anywhere else
+	FString Password = SPerforceSourceControlSettings::GetPassword();
+	if(!Password.IsEmpty())
 	{
-		OutConnectionInfo.Password = SPerforceSourceControlSettings::GetPassword();
+		OutConnectionInfo.Password = Password;
 	}
 
-	// Ticket is stored in the provider
+	// Ticket is stored in the provider (this is only set by the command line so should be safe to access without threading protection)
 	FPerforceSourceControlModule& PerforceSourceControl = FModuleManager::GetModuleChecked<FPerforceSourceControlModule>( "PerforceSourceControl" );
 	FPerforceSourceControlProvider& Provider = PerforceSourceControl.GetProvider();
 	OutConnectionInfo.Ticket = Provider.GetTicket();
diff --git a/Engine/Source/Developer/SourceControl/Private/ISourceControlProvider.cpp b/Engine/Source/Developer/SourceControl/Private/ISourceControlProvider.cpp
--- a/Engine/Source/Developer/SourceControl/Private/ISourceControlProvider.cpp
+++ b/Engine/Source/Developer/SourceControl/Private/ISourceControlProvider.cpp
@@ -98,4 +98,24 @@ TSharedPtr<class ISourceControlLabel> ISourceControlProvider::GetLabel(const FSt
 	return nullptr;
 }
 
+bool ISourceControlProvider::TryToDownloadFileFromBackgroundThread(const TSharedRef<class FDownloadFile>& InOperation, const FString& InFile)
+{
+	TArray<FString> FileArray;
+	FileArray.Add(InFile);
+
+	return TryToDownloadFileFromBackgroundThread(InOperation, FileArray);
+}
+
+bool ISourceControlProvider::TryToDownloadFileFromBackgroundThread(const TSharedRef<class FDownloadFile>& InOperation, const TArray<FString>& InFiles)
+{
+	// TryToDownloadFileFromBackgroundThread is unsupported by this source control provider
+	FFormatNamedArguments Arguments;
+	Arguments.Add(TEXT("ProviderName"), FText::FromName(GetName()));
+	FText Message = FText::Format(LOCTEXT("UnsupportedOperation", "TryToDownloadFileFromBackgroundThread is not supported by source control provider '{ProviderName}'"), Arguments);
+
+	InOperation->AddErrorMessge(Message);
+
+	return false;
+}
+
 #undef LOCTEXT_NAMESPACE
diff --git a/Engine/Source/Developer/SourceControl/Private/SourceControlOperations.cpp b/Engine/Source/Developer/SourceControl/Private/SourceControlOperations.cpp
new file mode 100644
--- /dev/null
+++ b/Engine/Source/Developer/SourceControl/Private/SourceControlOperations.cpp
@@ -0,0 +1,24 @@
+// Copyright Epic Games, Inc. All Rights Reserved.
+
+#include "SourceControlOperations.h"
+
+#include "Containers/StringView.h"
+#include "ISourceControlModule.h"
+#include "Misc/Paths.h"
+
+FDownloadFile::FDownloadFile(FStringView InTargetDirectory)
+	: TargetDirectory(InTargetDirectory)
+{
+	FPaths::NormalizeDirectoryName(TargetDirectory);
+
+	// Due to the asynchronous nature of the source control api, it might be some time before
+	// the TargetDirectory is actually used. So we do a validation pass on it now to try and 
+	// give errors close to the point that they were created. The caller is still free to try
+	// and use the FDownloadFile but with an invalid path it probably won't work.
+	FText Reason;
+	if (!FPaths::ValidatePath(TargetDirectory, &Reason))
+	{
+		UE_LOG(LogSourceControl, Error, TEXT("Path '%s' passed to FDownloadFile is invalid due to: %s"), 
+			*TargetDirectory, *Reason.ToString());
+	}
+}
diff --git a/Engine/Source/Developer/SourceControl/Public/ISourceControlOperation.h b/Engine/Source/Developer/SourceControl/Public/ISourceControlOperation.h
--- a/Engine/Source/Developer/SourceControl/Public/ISourceControlOperation.h
+++ b/Engine/Source/Developer/SourceControl/Public/ISourceControlOperation.h
@@ -69,6 +69,16 @@ public:
 		// Implemented in subclasses
 	}
 
+	/**
+	 * This will return true if the operation can be safely called from a background thread.
+	 * Currently it is assumed to only the operation 'FDownloadFile' will return true at least
+	 * until the API is made thread safe.
+	 */
+	virtual bool CanBeCalledFromBackgroundThreads() const
+	{
+		return false;
+	}
+
 	/** Factory method for easier operation creation */
 	template<typename Type, typename... TArgs>
 	static TSharedRef<Type, ESPMode::ThreadSafe> Create(TArgs&&... Args)
diff --git a/Engine/Source/Developer/SourceControl/Public/ISourceControlProvider.h b/Engine/Source/Developer/SourceControl/Public/ISourceControlProvider.h
--- a/Engine/Source/Developer/SourceControl/Public/ISourceControlProvider.h
+++ b/Engine/Source/Developer/SourceControl/Public/ISourceControlProvider.h
@@ -262,6 +262,32 @@ public:
 
 	virtual TArray<FSourceControlChangelistRef> GetChangelists( EStateCacheUsage::Type InStateCacheUsage ) = 0;
 
+	/**
+	 * Executes the FDownloadFile operation, but unlike the ::Execute method this can be called from a background thread, this
+	 * works because FDownloadFile is thread safe and it does not change the state of source control.
+	 * NOTE: This is only intended for use by the virtualization module and will be deprecated at some point in the future when
+	 * thread safety is built into the system.
+	 * NOTE: This is only implemented for the perforce source control provider with no plans to extend this to other providers.
+	 * 
+	 * @param InOperation The download operation to be executed.
+	 * @param InFile The depot path of a file to download.
+	 * @return True if the operation was a success, otherwise false.
+	 */
+	SOURCECONTROL_API bool TryToDownloadFileFromBackgroundThread(const TSharedRef<class FDownloadFile>& InOperation, const FString& InFile);
+
+	/**
+	 * Executes the FDownloadFile operation, but unlike the ::Execute method this can be called from a background thread, this
+	 * works because FDownloadFile is thread safe and it does not change the state of source control.
+	 * NOTE: This is only intended for use by the virtualization module and will be deprecated at some point in the future when
+	 * thread safety is built into the system.
+	 * NOTE: This is only implemented for the perforce source control provider with no plans to extend this to other providers.
+	 *
+	 * @param InOperation The download operation to be executed.
+	 * @param InFile An array of depot paths for the files to download.
+	 * @return True if the operation was a success, otherwise false.
+	 */
+	SOURCECONTROL_API virtual bool TryToDownloadFileFromBackgroundThread(const TSharedRef<class FDownloadFile>& InOperation, const TArray<FString>& InFiles);
+
 	/**
 	 * Whether the provider uses local read-only state to signal whether a file is editable.
 	 */
diff --git a/Engine/Source/Developer/SourceControl/Public/SourceControlOperations.h b/Engine/Source/Developer/SourceControl/Public/SourceControlOperations.h
--- a/Engine/Source/Developer/SourceControl/Public/SourceControlOperations.h
+++ b/Engine/Source/Developer/SourceControl/Public/SourceControlOperations.h
@@ -3,8 +3,11 @@
 #pragma once
 
 #include "CoreMinimal.h"
-#include "SourceControlOperationBase.h"
+
+#include "Containers/StringFwd.h"
+#include "Memory/SharedBuffer.h"
 #include "ISourceControlChangelist.h"
+#include "SourceControlOperationBase.h"
 
 #define LOCTEXT_NAMESPACE "SourceControl"
 
@@ -630,4 +633,64 @@ public:
 	}
 };
 
+/**
+ * Operation used to download a file from the source control server
+ */
+class FDownloadFile : public FSourceControlOperationBase
+{
+public:
+	/** 
+	 * This constructor will download the files and keep them in memory, 
+	 * which can then be accessed by calling XXX.
+	 */
+	FDownloadFile() = default;
+	
+	/** This constructor will download the files to the given directory */
+	SOURCECONTROL_API FDownloadFile(FStringView InTargetDirectory);
+
+	// ISourceControlOperation interface
+	virtual FName GetName() const override
+	{
+		return "DownloadFile";
+	}
+
+	virtual FText GetInProgressString() const override
+	{
+		return LOCTEXT("SourceControl_PrintOperation", "Downloading file from server...");
+	}
+
+	virtual bool CanBeCalledFromBackgroundThreads() const
+	{
+		return true;
+	}
+
+	FString GetTargetDirectory() const
+	{
+		return TargetDirectory;
+	}
+
+	void AddFileData(const FString& Filename, FSharedBuffer FileData)
+	{
+		FileDataMap.Add(Filename, FileData);
+	}
+
+	FSharedBuffer GetFileData(const FStringView& Filename)
+	{
+		const uint32 Hash = GetTypeHash(Filename);
+		FSharedBuffer* Buffer = FileDataMap.FindByHash(Hash, Filename);
+		if (Buffer != nullptr)
+		{
+			return *Buffer;
+		}
+		else
+		{
+			return FSharedBuffer();
+		}
+	}
+private:
+
+	FString TargetDirectory;
+	TMap<FString, FSharedBuffer> FileDataMap;
+};
+
 #undef LOCTEXT_NAMESPACE
