Allow the FSharedBuffer to take direct ownership of the byte array collected during perforce commands.

- This prevents the additional allocation, memcpy and spike in memory use.
- Made FP4ClientUser::DataBuffer private to force it's access via ::ReleaseData

#rnx

diff --git a/Engine/Plugins/Developer/PerforceSourceControl/Source/PerforceSourceControl/Private/PerforceConnection.cpp b/Engine/Plugins/Developer/PerforceSourceControl/Source/PerforceSourceControl/Private/PerforceConnection.cpp
--- a/Engine/Plugins/Developer/PerforceSourceControl/Source/PerforceSourceControl/Private/PerforceConnection.cpp
+++ b/Engine/Plugins/Developer/PerforceSourceControl/Source/PerforceSourceControl/Private/PerforceConnection.cpp
@@ -121,9 +121,17 @@ public:
 		return EnumHasAnyFlags(Flags, EP4ClientUserFlags::CollectData);
 	}
 
+	/** Returns DataBuffer as a FSharedBuffer, note that once called DataBuffer will be empty */
+	inline FSharedBuffer ReleaseData()
+	{
+		return MakeSharedBufferFromArray(MoveTemp(DataBuffer));
+	}
+
 	EP4ClientUserFlags Flags;
 	FP4RecordSet& Records;
 	TArray<FText>& OutErrorMessages;
+
+private:
 	TArray64<char> DataBuffer;
 };
 
@@ -813,7 +821,7 @@ bool FPerforceConnection::RunCommand(const FString& InCommand, const TArray<FStr
 
 	if (OutData)
 	{
-		OutData = FSharedBuffer::Clone(User.DataBuffer.GetData(), User.DataBuffer.Num());
+		OutData = User.ReleaseData();
 	}	
 
 	P4Client.SetBreak(NULL);
