Remove the defunct define 'USE_P4_API' from the PerforceSourceControl module.

#rnx

* Reasoning
- It was originally added as a way to tempoarily disble the perforce C/C++ api because we did not have access to VS2015 libs at the time.
- The code no longer compiles when USE_P4_API was set to 0 so would need to be fixed up to make it useful.
- We have no plans to disable the api making the define a bit pointless.

* Changes
- Now we can be sure that FP4ClientUser is derived from the real version of ClientUser we can mark the overriden methhods with virtiual and override keywords to make it more obvious.
- Removed some pointless comments, like // Constructor.
- Added comments to overriden methods that were missing them.
- Changed some // comments to /** */ to keep consistency in the code.
- Replaced occurences of NULL with nullptr.

diff --git a/Engine/Plugins/Developer/PerforceSourceControl/Source/PerforceSourceControl/PerforceSourceControl.Build.cs b/Engine/Plugins/Developer/PerforceSourceControl/Source/PerforceSourceControl/PerforceSourceControl.Build.cs
--- a/Engine/Plugins/Developer/PerforceSourceControl/Source/PerforceSourceControl/PerforceSourceControl.Build.cs
+++ b/Engine/Plugins/Developer/PerforceSourceControl/Source/PerforceSourceControl/PerforceSourceControl.Build.cs
@@ -19,7 +19,6 @@ public class PerforceSourceControl : ModuleRules
 		);
 
 		AddEngineThirdPartyPrivateStaticDependencies(Target, "Perforce");
-		PublicDefinitions.Add("USE_P4_API=1");
 
 		if (Target.Platform == UnrealTargetPlatform.Win64 || Target.Platform == UnrealTargetPlatform.Mac)
 		{
diff --git a/Engine/Plugins/Developer/PerforceSourceControl/Source/PerforceSourceControl/Private/PerforceConnection.cpp b/Engine/Plugins/Developer/PerforceSourceControl/Source/PerforceSourceControl/Private/PerforceConnection.cpp
--- a/Engine/Plugins/Developer/PerforceSourceControl/Source/PerforceSourceControl/Private/PerforceConnection.cpp
+++ b/Engine/Plugins/Developer/PerforceSourceControl/Source/PerforceSourceControl/Private/PerforceConnection.cpp
@@ -33,7 +33,6 @@ namespace
 class FP4ClientUser : public ClientUser
 {
 public:
-	
 	FP4ClientUser(FP4RecordSet& InRecords, EP4ClientUserFlags InFlags, TArray<FText>& InOutErrorMessages)
 		: ClientUser()
 		, Flags(InFlags)
@@ -44,9 +43,8 @@ public:
 	}
 
 	/**  Called by P4API when the results from running a command are ready. */
-	void OutputStat(StrDict* VarList)
+	virtual void OutputStat(StrDict* VarList) override
 	{
-#if USE_P4_API
 		FP4Record Record;
 		StrRef Var, Value;
 
@@ -72,11 +70,10 @@ public:
 		}
 
 		Records.Add(Record);
-#endif
 	}
 
 	/** Called by P4API when it output a chunk of text data from a file (commonly via P4 Print) */
-	void OutputText(const char* DataPtr, int DataLength)
+	virtual void OutputText(const char* DataPtr, int DataLength) override
 	{
 		if (IsCollectingData())
 		{
@@ -102,13 +99,11 @@ public:
 	}
 
 	/** Called by P4API when an error message is avaliable. */
-	void HandleError(Error* InError)
+	virtual void HandleError(Error* InError) override
 	{
-#if USE_P4_API
 		StrBuf ErrorMessage;
 		InError->Fmt(&ErrorMessage);
 		OutErrorMessages.Add(FText::FromString(FString(TO_TCHAR(ErrorMessage.Text(), IsUnicodeServer()))));
-#endif
 	}
 
 	inline bool IsUnicodeServer() const
@@ -139,8 +134,6 @@ private:
 class FP4CreateChangelistClientUser : public FP4ClientUser
 {
 public:
-
-	// Constructor
 	FP4CreateChangelistClientUser(FP4RecordSet& InRecords, EP4ClientUserFlags InFlags, TArray<FText>& InOutErrorMessages, const FText& InDescription, const TArray<FString>& InFiles, ClientApi &InP4Client)
 		:	FP4ClientUser(InRecords, InFlags, InOutErrorMessages)
 		,	Description(InDescription)
@@ -150,8 +143,8 @@ public:
 	{
 	}
 
-	// Called by P4API when the changelist is created.
-	void OutputInfo(ANSICHAR Level, const ANSICHAR *Data)
+	/** Called by P4API when the changelist is created. */
+	virtual void OutputInfo(ANSICHAR Level, const ANSICHAR *Data) override
 	{
 		const int32 ChangeTextLen = FCString::Strlen(TEXT("Change "));
 		if (FString(TO_TCHAR(Data, IsUnicodeServer())).StartsWith(TEXT("Change ")))
@@ -160,10 +153,9 @@ public:
 		}
 	}
 
-	// Called by P4API on "change -i" command. OutBuffer is filled with changelist specification text
-	void InputData(StrBuf* OutBuffer, Error* OutError)
+	/** Called by P4API on "change -i" command. OutBuffer is filled with changelist specification text. */
+	virtual void InputData(StrBuf* OutBuffer, Error* OutError) override
 	{
-#if USE_P4_API
 		FString OutputDesc;
 		OutputDesc += TEXT("Change:\tnew\n\n");
 		OutputDesc += TEXT("Client:\t");
@@ -195,7 +187,6 @@ public:
 		OutputDesc += TEXT("\n");
 
 		OutBuffer->Append(FROM_TCHAR(*OutputDesc, IsUnicodeServer()));
-#endif
 	}
 
 	FText Description;
@@ -208,7 +199,6 @@ public:
 class FP4EditChangelistClientUser : public FP4ClientUser
 {
 public:
-	// Constructor
 	FP4EditChangelistClientUser(FP4RecordSet& OutRecords, EP4ClientUserFlags InFlags, TArray<FText>& InOutErrorMessages, const FText& InDescription, int32 InChangelistNumber, const FP4RecordSet& InRecords, ClientApi& InP4Client)
 		: FP4ClientUser(OutRecords, InFlags, InOutErrorMessages)
 		, Description(InDescription)
@@ -218,8 +208,8 @@ public:
 	{
 	}
 
-	// Called by P4API when the changelist is updated.
-	void OutputInfo(ANSICHAR Level, const ANSICHAR* Data)
+	/** Called by P4API when the changelist is updated. */
+	virtual void OutputInfo(ANSICHAR Level, const ANSICHAR* Data) override
 	{
 		const int32 ChangeTextLen = FCString::Strlen(TEXT("Change "));
 		if (FString(TO_TCHAR(Data, IsUnicodeServer())).StartsWith(TEXT("Change ")))
@@ -228,10 +218,9 @@ public:
 		}
 	}
 
-	// Called by P4API on "change -i" command. OutBuffer is filled with changelist specification text
-	void InputData(StrBuf* OutBuffer, Error* OutError)
+	/** Called by P4API on "change -i" command. OutBuffer is filled with changelist specification text. */
+	virtual void InputData(StrBuf* OutBuffer, Error* OutError) override
 	{
-#if USE_P4_API
 		FString OutputDesc;
 		const FP4Record& Record = PreviousRecords[0];
 
@@ -289,7 +278,6 @@ public:
 		}
 
 		OutBuffer->Append(FROM_TCHAR(*OutputDesc, IsUnicodeServer()));
-#endif
 	}
 
 	FText Description;
@@ -302,19 +290,16 @@ public:
 class FP4LoginClientUser : public FP4ClientUser
 {
 public:
-
-	// Constructor
 	FP4LoginClientUser(const FString& InPassword, FP4RecordSet& InRecords, EP4ClientUserFlags InFlags, TArray<FText>& InOutErrorMessages)
 		:	FP4ClientUser(InRecords, InFlags, InOutErrorMessages)
 		,	Password(InPassword)
 	{
 	}
 
-	void Prompt( const StrPtr& InMessage, StrBuf& OutPrompt, int NoEcho, Error* InError )
+	/** Called when we are prompted for a password */
+	virtual void Prompt( const StrPtr& InMessage, StrBuf& OutPrompt, int NoEcho, Error* InError ) override
 	{
-#if USE_P4_API
 		OutPrompt.Set(FROM_TCHAR(*Password, IsUnicodeServer()));
-#endif
 	}
 
 	/** Password to use when logging in */
@@ -330,7 +315,7 @@ public:
 	{
 	}
 
-#if USE_P4_API
+	/** Called when the perforce connection wants to know if it should stay connected */
 	virtual int IsAlive() override
 	{
 		if(IsCancelled.IsBound() && IsCancelled.Execute())
@@ -340,7 +325,6 @@ public:
 
 		return 1;
 	}
-#endif
 
 	FOnIsCancelled IsCancelled;
 };
@@ -369,11 +353,9 @@ static bool TestConnection(ClientApi& P4Client, const FString& ClientSpecName, b
 	}
 
 	const char *ArgV[] = { "-o", (char*)ClientSpecUTF8Name };
-#if USE_P4_API
-	P4Client.SetArgv(2, const_cast<char*const*>(ArgV));
 
+	P4Client.SetArgv(2, const_cast<char*const*>(ArgV));
 	P4Client.Run("client", &User);
-#endif
 
 	// clean up args
 	delete [] ClientSpecUTF8Name;
@@ -393,7 +375,6 @@ static bool TestConnection(ClientApi& P4Client, const FString& ClientSpecName, b
 
 static bool CheckUnicodeStatus(ClientApi& P4Client, bool& bIsUnicodeServer, TArray<FText>& OutErrorMessages)
 {
-#if USE_P4_API
 	FP4RecordSet Records;
 	FP4ClientUser User(Records, EP4ClientUserFlags::None, OutErrorMessages);
 
@@ -401,13 +382,12 @@ static bool CheckUnicodeStatus(ClientApi& P4Client, bool& bIsUnicodeServer, TArr
 
 	if(Records.Num() > 0)
 	{
-		bIsUnicodeServer = Records[0].Find(TEXT("unicode")) != NULL;
+		bIsUnicodeServer = Records[0].Find(TEXT("unicode")) != nullptr;
 	}
 	else
 	{
 		bIsUnicodeServer = false;
 	}
-#endif
 
 	return OutErrorMessages.Num() == 0;
 }
@@ -466,7 +446,7 @@ bool FPerforceConnection::AutoDetectWorkspace(const FPerforceConnectionInfo& InC
 bool FPerforceConnection::Login(const FPerforceConnectionInfo& InConnectionInfo)
 {
 	TArray<FText> ErrorMessages;
-#if USE_P4_API
+
 	FP4RecordSet Records;
 	FP4LoginClientUser User(InConnectionInfo.Password, Records, EP4ClientUserFlags::None, ErrorMessages);
 
@@ -482,7 +462,6 @@ bool FPerforceConnection::Login(const FPerforceConnectionInfo& InConnectionInfo)
 			UE_LOG(LogSourceControl, Error, TEXT("%s"), *ErrorMessage.ToString());
 		}
 	}
-#endif
 
 	return ErrorMessages.Num() == 0;
 }
@@ -493,7 +472,7 @@ bool FPerforceConnection::EnsureValidConnection(FString& InOutServerName, FStrin
 
 	bool bIsUnicodeServer = false;
 	bool bConnectionOK = false;
-#if USE_P4_API
+
 	FMessageLog SourceControlLog("SourceControl");
 
 	FString NewServerName = InOutServerName;
@@ -643,7 +622,6 @@ bool FPerforceConnection::EnsureValidConnection(FString& InOutServerName, FStrin
 		InOutUserName = NewUserName;
 		InOutWorkspaceName = NewClientSpecName;
 	}
-#endif
 
 	return bConnectionOK;
 }
@@ -736,16 +714,11 @@ bool FPerforceConnection::GetWorkspaceList(const FPerforceConnectionInfo& InConn
 
 bool FPerforceConnection::IsValidConnection()
 {
-#if USE_P4_API
 	return bEstablishedConnection && !P4Client.Dropped();
-#else
-	return false;
-#endif
 }
 
 void FPerforceConnection::Disconnect()
 {
-#if USE_P4_API
 	Error P4Error;
 
 	P4Client.Final(&P4Error);
@@ -757,12 +730,10 @@ void FPerforceConnection::Disconnect()
 		UE_LOG(LogSourceControl, Error, TEXT("P4ERROR: Failed to disconnect from Server."));
 		UE_LOG(LogSourceControl, Error, TEXT("%s"), TO_TCHAR(ErrorMessage.Text(), bIsUnicode));
 	}
-#endif
 }
 
 bool FPerforceConnection::RunCommand(const FString& InCommand, const TArray<FString>& InParameters, FP4RecordSet& OutRecordSet, TOptional<FSharedBuffer>& OutData, TArray<FText>& OutErrorMessage, FOnIsCancelled InIsCancelled, bool& OutConnectionDropped, const bool bInStandardDebugOutput, const bool bInAllowRetry)
 {
-#if USE_P4_API
 	if (!bEstablishedConnection)
 	{
 		return false;
@@ -824,7 +795,7 @@ bool FPerforceConnection::RunCommand(const FString& InCommand, const TArray<FStr
 		OutData = User.ReleaseData();
 	}	
 
-	P4Client.SetBreak(NULL);
+	P4Client.SetBreak(nullptr);
 
 	// Free arguments
 	for (int32 Index = 0; Index < ArgC; Index++)
@@ -843,14 +814,12 @@ bool FPerforceConnection::RunCommand(const FString& InCommand, const TArray<FStr
 	{
 		UE_LOG( LogSourceControl, VeryVerbose, TEXT("P4 execution time: %0.4f seconds. Command: %s"), FPlatformTime::Seconds() - SCCStartTime, *FullCommand );
 	}
-#endif
 
 	return OutRecordSet.Num() > 0;
 }
 
 int32 FPerforceConnection::CreatePendingChangelist(const FText &Description, const TArray<FString>& Files, FOnIsCancelled InIsCancelled, TArray<FText>& OutErrorMessages)
 {
-#if USE_P4_API
 	TArray<FString> Params;
 	FP4RecordSet Records;
 	EP4ClientUserFlags Flags = bIsUnicode ? EP4ClientUserFlags::UnicodeServer : EP4ClientUserFlags::None;
@@ -864,17 +833,13 @@ int32 FPerforceConnection::CreatePendingChangelist(const FText &Description, con
 	FP4CreateChangelistClientUser User(Records, Flags, OutErrorMessages, Description, Files, P4Client);
 	P4Client.Run("change", &User);
 
-	P4Client.SetBreak(NULL);
+	P4Client.SetBreak(nullptr);
 
 	return User.ChangelistNumber;
-#else
-	return 0;
-#endif
 }
 
 int32 FPerforceConnection::EditPendingChangelist(const FText& NewDescription, int32 ChangelistNumber, FOnIsCancelled InIsCancelled, TArray<FText>& OutErrorMessages)
 {
-#if USE_P4_API
 	FP4RecordSet PreviousRecords;
 
 	// Get changelist current specification
@@ -910,18 +875,14 @@ int32 FPerforceConnection::EditPendingChangelist(const FText& NewDescription, in
 		FP4EditChangelistClientUser User(Records, Flags, OutErrorMessages, NewDescription, ChangelistNumber, PreviousRecords, P4Client);
 		P4Client.Run("change", &User);
 
-		P4Client.SetBreak(NULL);
+		P4Client.SetBreak(nullptr);
 
 		return (OutErrorMessages.Num() == 0 ? User.ChangelistNumber : 0);
 	}
-#else
-	return 0;
-#endif
 }
 
 void FPerforceConnection::EstablishConnection(const FPerforceConnectionInfo& InConnectionInfo)
 {
-#if USE_P4_API
 	// Verify Input. ServerName and UserName are required
 	if ( InConnectionInfo.Port.IsEmpty() || InConnectionInfo.UserName.IsEmpty() )
 	{
@@ -978,7 +939,7 @@ void FPerforceConnection::EstablishConnection(const FPerforceConnectionInfo& InC
 		if (RunCommand(TEXT("info"), Params, Records, CommandData, ErrorMessages, FOnIsCancelled(), bConnectionDropped, bStandardDebugOutput, bAllowRetry))
 		{
 			// Get character encoding
-			bIsUnicode = Records[0].Find(TEXT("unicode")) != NULL;
+			bIsUnicode = Records[0].Find(TEXT("unicode")) != nullptr;
 			if(bIsUnicode)
 			{
 				P4Client.SetTrans(CharSetApi::UTF_8);
@@ -1018,11 +979,10 @@ void FPerforceConnection::EstablishConnection(const FPerforceConnectionInfo& InC
 			}
 		}
 	}
-#endif
 }
 
 FScopedPerforceConnection::FScopedPerforceConnection( FPerforceSourceControlCommand& InCommand )
-	: Connection(NULL)
+	: Connection(nullptr)
 	, Concurrency(InCommand.Concurrency)
 {
 	Initialize(InCommand.ConnectionInfo);
@@ -1033,7 +993,7 @@ FScopedPerforceConnection::FScopedPerforceConnection( FPerforceSourceControlComm
 }
 
 FScopedPerforceConnection::FScopedPerforceConnection( EConcurrency::Type InConcurrency, const FPerforceConnectionInfo& InConnectionInfo )
-	: Connection(NULL)
+	: Connection(nullptr)
 	, Concurrency(InConcurrency)
 {
 	Initialize(InConnectionInfo);
@@ -1065,7 +1025,7 @@ void FScopedPerforceConnection::Initialize( const FPerforceConnectionInfo& InCon
 
 FScopedPerforceConnection::~FScopedPerforceConnection()
 {
-	if(Connection != NULL)
+	if(Connection != nullptr)
 	{
 		if(Concurrency == EConcurrency::Asynchronous)
 		{
@@ -1074,7 +1034,7 @@ FScopedPerforceConnection::~FScopedPerforceConnection()
 			delete Connection;
 		}
 		
-		Connection = NULL;
+		Connection = nullptr;
 	}
 }
 
diff --git a/Engine/Plugins/Developer/PerforceSourceControl/Source/PerforceSourceControl/Private/PerforceConnection.h b/Engine/Plugins/Developer/PerforceSourceControl/Source/PerforceSourceControl/Private/PerforceConnection.h
--- a/Engine/Plugins/Developer/PerforceSourceControl/Source/PerforceSourceControl/Private/PerforceConnection.h
+++ b/Engine/Plugins/Developer/PerforceSourceControl/Source/PerforceSourceControl/Private/PerforceConnection.h
@@ -110,10 +110,8 @@ public:
 	void EstablishConnection(const FPerforceConnectionInfo& InConnectionInfo);
 
 public:
-#if USE_P4_API
 	/** Perforce API client object */
 	ClientApi P4Client;
-#endif
 
 	/** The current root for the client workspace */
 	FString ClientRoot;
diff --git a/Engine/Plugins/Developer/PerforceSourceControl/Source/PerforceSourceControl/Private/PerforceSourceControlOperations.cpp b/Engine/Plugins/Developer/PerforceSourceControl/Source/PerforceSourceControl/Private/PerforceSourceControlOperations.cpp
--- a/Engine/Plugins/Developer/PerforceSourceControl/Source/PerforceSourceControl/Private/PerforceSourceControlOperations.cpp
+++ b/Engine/Plugins/Developer/PerforceSourceControl/Source/PerforceSourceControl/Private/PerforceSourceControlOperations.cpp
@@ -1534,7 +1534,6 @@ FName FPerforceUpdateStatusWorker::GetName() const
 
 bool FPerforceUpdateStatusWorker::Execute(FPerforceSourceControlCommand& InCommand)
 {
-#if USE_P4_API
 	FScopedPerforceConnection ScopedConnection(InCommand);
 	if (!InCommand.IsCanceled() && ScopedConnection.IsValid())
 	{
@@ -1674,7 +1673,6 @@ bool FPerforceUpdateStatusWorker::Execute(FPerforceSourceControlCommand& InComma
 			RemoveRedundantErrors(InCommand, TEXT(" - file(s) not opened on this client"));
 		}
 	}
-#endif
 
 	return InCommand.bCommandSuccessful;
 }
diff --git a/Engine/Plugins/Developer/PerforceSourceControl/Source/PerforceSourceControl/Private/PerforceSourceControlPrivate.h b/Engine/Plugins/Developer/PerforceSourceControl/Source/PerforceSourceControl/Private/PerforceSourceControlPrivate.h
--- a/Engine/Plugins/Developer/PerforceSourceControl/Source/PerforceSourceControl/Private/PerforceSourceControlPrivate.h
+++ b/Engine/Plugins/Developer/PerforceSourceControl/Source/PerforceSourceControl/Private/PerforceSourceControlPrivate.h
@@ -4,29 +4,7 @@
 
 #include "CoreMinimal.h"
 
-//@todo: perforce doesn't currently support VS2015
-#if USE_P4_API
-PRAGMA_DISABLE_SHADOW_VARIABLE_WARNINGS
+THIRD_PARTY_INCLUDES_START
 #include <p4/clientapi.h>
 #include <p4/i18napi.h>
-PRAGMA_ENABLE_SHADOW_VARIABLE_WARNINGS
-#endif
-
-//@todo: perforce doesn't currently support VS2015
-#if !USE_P4_API
-class ClientApi
-{};
-class ClientUser
-{};
-class StrBuf
-{};
-class StrDict
-{};
-class Error
-{};
-class StrPtr
-{};
-class KeepAlive
-{};
-#endif
-
+THIRD_PARTY_INCLUDES_END
diff --git a/Engine/Plugins/Developer/PerforceSourceControl/Source/PerforceSourceControl/Private/SPerforceSourceControlSettings.cpp b/Engine/Plugins/Developer/PerforceSourceControl/Source/PerforceSourceControl/Private/SPerforceSourceControlSettings.cpp
--- a/Engine/Plugins/Developer/PerforceSourceControl/Source/PerforceSourceControl/Private/SPerforceSourceControlSettings.cpp
+++ b/Engine/Plugins/Developer/PerforceSourceControl/Source/PerforceSourceControl/Private/SPerforceSourceControlSettings.cpp
@@ -36,7 +36,6 @@ void SPerforceSourceControlSettings::Construct(const FArguments& InArgs)
 
 	if (PortName.IsEmpty() && UserName.IsEmpty())
 	{
-#if USE_P4_API
 		ClientApi TestP4;
 		TestP4.SetProg("UE");
 		Error P4Error;
@@ -48,7 +47,6 @@ void SPerforceSourceControlSettings::Construct(const FArguments& InArgs)
 		PerforceSourceControl.AccessSettings().SetPort(PortName);
 		PerforceSourceControl.AccessSettings().SetUserName(UserName);
 		PerforceSourceControl.SaveSettings();
-#endif
 	}
 
 	ChildSlot
