Unshelved from pending changelist '16239478':

PR #7989: Fixed issue where file history was truncated for moved files (Contributed by [REDACTED_PERSON])

#github

diff --git a/Engine/Plugins/Developer/PerforceSourceControl/Source/PerforceSourceControl/Private/PerforceSourceControlOperations.cpp b/Engine/Plugins/Developer/PerforceSourceControl/Source/PerforceSourceControl/Private/PerforceSourceControlOperations.cpp
--- a/Engine/Plugins/Developer/PerforceSourceControl/Source/PerforceSourceControl/Private/PerforceSourceControlOperations.cpp
+++ b/Engine/Plugins/Developer/PerforceSourceControl/Source/PerforceSourceControl/Private/PerforceSourceControlOperations.cpp
@@ -1437,6 +1437,12 @@ static void ParseHistoryResults(const FP4RecordSet& InRecords, const TArray<FPer
 			FString DepotFileName = ClientRecord(TEXT("depotFile"));
 			FString LocalFileName = FindWorkspaceFile(InStates, DepotFileName);
 
+			if (OutHistory.Contains(LocalFileName))
+			{
+				// If a file existed in a location multiple times we rely on the first, most recent record having all the revisions we are interested in
+				continue;
+			}
+
 			TArray< TSharedRef<FPerforceSourceControlRevision, ESPMode::ThreadSafe> > Revisions;
 			int32 RevisionNumbers = 0;
 			for (;;)
