Adds "LastSynced" flag to Perforce Sync operation, allowing to force sync (ie revert) to the last synced version of an asset.
Modifies Revert operation on Uncontrolled assets to force sync to last synced version instead of head revision.

#changelist validated

diff --git a/Engine/Plugins/Developer/PerforceSourceControl/Source/PerforceSourceControl/Private/PerforceSourceControlOperations.cpp b/Engine/Plugins/Developer/PerforceSourceControl/Source/PerforceSourceControl/Private/PerforceSourceControlOperations.cpp
--- a/Engine/Plugins/Developer/PerforceSourceControl/Source/PerforceSourceControl/Private/PerforceSourceControlOperations.cpp
+++ b/Engine/Plugins/Developer/PerforceSourceControl/Source/PerforceSourceControl/Private/PerforceSourceControlOperations.cpp
@@ -791,6 +791,8 @@ bool FPerforceSyncWorker::Execute(FPerforceSourceControlCommand& InCommand)
 	if (!InCommand.IsCanceled() && ScopedConnection.IsValid())
 	{
 		FPerforceConnection& Connection = ScopedConnection.GetConnection();
+		FPerforceSourceControlModule& PerforceSourceControl = FPerforceSourceControlModule::Get();
+
 		TSharedRef<FSync, ESPMode::ThreadSafe> Operation = StaticCastSharedRef<FSync>(InCommand.Operation);
 		TArray<FString> Parameters;
 
@@ -813,6 +815,22 @@ bool FPerforceSyncWorker::Execute(FPerforceSourceControlCommand& InCommand)
 			{
 				FileName += TEXT("#head");
 			}
+			else if (Operation->IsLastSyncedFlagSet())
+			{
+				TSharedRef<FPerforceSourceControlState, ESPMode::ThreadSafe> State = PerforceSourceControl.GetProvider().GetStateInternal(FileName);
+
+				if (State->LocalRevNumber == ISourceControlState::INVALID_REVISION)
+				{
+					const FString ErrorMsg = FString::Format(TEXT("Cached state of {0} is missing. Please refresh and try again."), {FileName});
+					
+					InCommand.ResultInfo.ErrorMessages.Add(FText::FromString(ErrorMsg));
+					InCommand.bCommandSuccessful = false;
+					
+					return InCommand.bCommandSuccessful;
+				}
+
+				FileName += FString::Format(TEXT("#{0}"), { State->LocalRevNumber });
+			}
 			else if (!Revision.IsEmpty())
 			{
 				// @= syncs the file to the submitted/shelved changelist number
diff --git a/Engine/Source/Developer/SourceControl/Public/SourceControlOperations.h b/Engine/Source/Developer/SourceControl/Public/SourceControlOperations.h
--- a/Engine/Source/Developer/SourceControl/Public/SourceControlOperations.h
+++ b/Engine/Source/Developer/SourceControl/Public/SourceControlOperations.h
@@ -238,13 +238,26 @@ public:
 		return bForce;
 	}
 
+	void SetLastSyncedFlag(const bool bInLastSynced)
+	{
+		bLastSynced = bInLastSynced;
+	}
+
+	bool IsLastSyncedFlagSet() const
+	{
+		return bLastSynced;
+	}
+
 protected:
 	/** Revision to sync to */
 	FString Revision;
 
-	/** Flag abstracting if the operation aim to sync to head */
+	/** Flag abstracting if the operation aims to sync to head */
 	bool bHeadRevision = false;
 
+	/** Flag abstracting if the operation aims to force sync to last synced revision */
+	bool bLastSynced = false;
+
 	/** Forces operation, even if the file is already at the wanted revision. */
 	bool bForce = false;
 };
diff --git a/Engine/Source/Editor/SourceControlWindows/Private/SSourceControlChangelists.cpp b/Engine/Source/Editor/SourceControlWindows/Private/SSourceControlChangelists.cpp
--- a/Engine/Source/Editor/SourceControlWindows/Private/SSourceControlChangelists.cpp
+++ b/Engine/Source/Editor/SourceControlWindows/Private/SSourceControlChangelists.cpp
@@ -757,28 +757,20 @@ void SSourceControlChangelistsWidget::OnRevert()
 		SourceControlProvider.Execute(RevertOperation, SelectedChangelist, SelectedControlledFiles);
 	}
 
-	// Reverts selected Uncontrolled Files
-	if (!SelectedUncontrolledFiles.IsEmpty())
-	{
-		auto ForceSyncOperation = ISourceControlOperation::Create<FSync>();
-		ForceSyncOperation->SetForce(true);
-		ForceSyncOperation->SetHeadRevisionFlag(true);
-		SourceControlProvider.Execute(ForceSyncOperation, SelectedUncontrolledFiles);
-
-		FUncontrolledChangelistsModule::Get().UpdateStatus();
-	}
-
 	FUncontrolledChangelistStatePtr SelectedUncontrolledChangelist = GetCurrentUncontrolledChangelistState();
 
 	// Reverts the selected Uncontrolled Changelist
 	if (SelectedUncontrolledChangelist.IsValid())
 	{
-		SelectedUncontrolledFiles.Reset();
 		Algo::Transform(SelectedUncontrolledChangelist->GetFilesStates(), SelectedUncontrolledFiles, [](const FSourceControlStateRef& State) { return State->GetFilename(); });
+	}
 
+	// Reverts selected Uncontrolled Files
+	if (!SelectedUncontrolledFiles.IsEmpty())
+	{
 		auto ForceSyncOperation = ISourceControlOperation::Create<FSync>();
 		ForceSyncOperation->SetForce(true);
-		ForceSyncOperation->SetHeadRevisionFlag(true);
+		ForceSyncOperation->SetLastSyncedFlag(true);
 		SourceControlProvider.Execute(ForceSyncOperation, SelectedUncontrolledFiles);
 
 		FUncontrolledChangelistsModule::Get().UpdateStatus();
