The ability to create, delete and switch workspaces via the SourceControl API. Currently this is only implemented for the perforce source control plugin.

#rnx

### ISourceControlProvider / PerforceSourceControlProvider
- Add new method ::SwitchWorkspace allowing the caller to switch the current workspace.
    - Currently only implemented for perforce.
    - Calling this will cause the provider to close and then re-open it's connection and so will invalidate any existing established connection or cached states. Due to this the method will fail if there are outstanding commands being processed.
    - It is intended that this is only used for short periods of time, in order to submit or access files from locations that the user chosen workspace might be able to access.
    - The new workspace will not be saved to the source control settings. Although this could occur later if done.

### SourceControlOperations
- Add new operation **FCreateWorkspace**.
    - Currently only implemented for perforce.
    - The name and the root of the workspace (filepath on the client machine) must be provided in the constructor.
    - Additional mappings can be provided by calling **AddNativeClientViewMapping** which expects the mapping to be in the native format of the target source control system. At some point in the future we may add a more generic, target agnostic, format.
- Add new operation **FDeleteWorkspace**.
    - Currently only implemented for perforce.
   - Deletes the workspace of the given name although the workspace must be in a state where it can be deleted (no files open for add/edit/delete etc) or the operation will fail.

### PerforceConnection
- Add a new method ::CreateWorkspace for creating workspaces
- Add a new extended client 'FP4CommandWithStdInputClientUser'
    - We should really extend FP4ClientUser to be able to take overridden stdinput instead  but that would mean adding yet another input parameter. I'd rather keep the code separate for now and then merge them together in a future refactor where the many input parameters of ::RunCommand is replaces with a struct with TOptional members.
- Added a ::GetUser method, which returns the user of the current connection without the caller needing to know if the connection is in unicode format or not.

diff --git a/Engine/Plugins/Developer/PerforceSourceControl/Source/PerforceSourceControl/Private/PerforceConnection.cpp b/Engine/Plugins/Developer/PerforceSourceControl/Source/PerforceSourceControl/Private/PerforceConnection.cpp
--- a/Engine/Plugins/Developer/PerforceSourceControl/Source/PerforceSourceControl/Private/PerforceConnection.cpp
+++ b/Engine/Plugins/Developer/PerforceSourceControl/Source/PerforceSourceControl/Private/PerforceConnection.cpp
@@ -286,6 +286,28 @@ public:
 	FP4RecordSet PreviousRecords;
 };
 
+/** An extended version of FP4ClientUser that allows the standard input stream to be overriden */
+class FP4CommandWithStdInputClientUser : public FP4ClientUser
+{
+public:
+	FP4CommandWithStdInputClientUser(FStringView InStdInput, FP4RecordSet& InRecords, EP4ClientUserFlags InFlags, TArray<FText>& InOutErrorMessages, ClientApi& InP4Client)
+		: FP4ClientUser(InRecords, InFlags, InOutErrorMessages)
+		, P4Client(InP4Client)
+		, StdInput(InStdInput)
+	{
+	}
+
+
+	virtual void InputData(StrBuf* OutBuffer, Error* OutError) override
+	{	
+		OutBuffer->Append(FROM_TCHAR(*StdInput, IsUnicodeServer()));
+	}
+
+protected:
+	ClientApi& P4Client;
+	FString StdInput;
+};
+
 /** Custom ClientUser class for handling login commands */
 class FP4LoginClientUser : public FP4ClientUser
 {
@@ -881,6 +903,27 @@ int32 FPerforceConnection::EditPendingChangelist(const FText& NewDescription, in
 	}
 }
 
+bool FPerforceConnection::CreateWorkspace(FStringView WorkspaceSpec, FOnIsCancelled InIsCancelled, TArray<FText>& OutErrorMessages)
+{
+	TArray<FString> Params;
+	FP4RecordSet Records;
+
+	EP4ClientUserFlags Flags = bIsUnicode ? EP4ClientUserFlags::UnicodeServer : EP4ClientUserFlags::None;
+
+	const char* ArgV[] = { "-i" };
+	P4Client.SetArgv(1, const_cast<char* const*>(ArgV));
+
+	FP4KeepAlive KeepAlive(InIsCancelled);
+	P4Client.SetBreak(&KeepAlive);
+
+	FP4CommandWithStdInputClientUser User(WorkspaceSpec, Records, Flags, OutErrorMessages, P4Client);
+	P4Client.Run("client", &User);
+
+	P4Client.SetBreak(nullptr);
+	
+	return OutErrorMessages.Num() == 0;
+}
+
 void FPerforceConnection::EstablishConnection(const FPerforceConnectionInfo& InConnectionInfo)
 {
 	// Verify Input. ServerName and UserName are required
@@ -981,6 +1024,18 @@ void FPerforceConnection::EstablishConnection(const FPerforceConnectionInfo& InC
 	}
 }
 
+FString FPerforceConnection::GetUser()
+{
+	if (bEstablishedConnection)
+	{
+		return TO_TCHAR(P4Client.GetUser().Text(), bIsUnicode);
+	}
+	else
+	{
+		return FString();
+	}
+}
+
 FScopedPerforceConnection::FScopedPerforceConnection( FPerforceSourceControlCommand& InCommand )
 	: Connection(nullptr)
 	, Concurrency(InCommand.Concurrency)
diff --git a/Engine/Plugins/Developer/PerforceSourceControl/Source/PerforceSourceControl/Private/PerforceConnection.h b/Engine/Plugins/Developer/PerforceSourceControl/Source/PerforceSourceControl/Private/PerforceConnection.h
--- a/Engine/Plugins/Developer/PerforceSourceControl/Source/PerforceSourceControl/Private/PerforceConnection.h
+++ b/Engine/Plugins/Developer/PerforceSourceControl/Source/PerforceSourceControl/Private/PerforceConnection.h
@@ -96,6 +96,17 @@ public:
 	 */
 	int32 EditPendingChangelist(const FText& NewDescription, int32 ChangelistNumber, FOnIsCancelled InIsCancelled, TArray<FText>& OutErrorMessages);
 
+	/** 
+	 * Creates a workspace based on the spec provided via the WorkspaceSpec.
+	 * 
+	 * @param WorkspaceSpec		The specification of the workspace to create
+	 * @param InIsCancelled		Delegate allowing the cancelling of the command if needed
+	 * @param OutErrorMessages	An array that will be filled with all errors encountered during the command
+	 * 
+	 * @return Returns true if no errors were encountered, otherwise false
+	 */
+	bool CreateWorkspace(FStringView WorkspaceSpec, FOnIsCancelled InIsCancelled, TArray<FText>& OutErrorMessages);
+
 	/**
 	 * Attempt to login - some servers will require this 
 	 * @param	InConnectionInfo		Credentials to use
@@ -109,6 +120,13 @@ public:
 	 */
 	void EstablishConnection(const FPerforceConnectionInfo& InConnectionInfo);
 
+	/** 
+	 * Return the user of the connection
+	 * 
+	 * @return The current perforce user, this will be blank if the connect failed.
+	 */
+	FString GetUser();
+
 public:
 	/** Perforce API client object */
 	ClientApi P4Client;
diff --git a/Engine/Plugins/Developer/PerforceSourceControl/Source/PerforceSourceControl/Private/PerforceSourceControlModule.cpp b/Engine/Plugins/Developer/PerforceSourceControl/Source/PerforceSourceControl/Private/PerforceSourceControlModule.cpp
--- a/Engine/Plugins/Developer/PerforceSourceControl/Source/PerforceSourceControl/Private/PerforceSourceControlModule.cpp
+++ b/Engine/Plugins/Developer/PerforceSourceControl/Source/PerforceSourceControl/Private/PerforceSourceControlModule.cpp
@@ -38,6 +38,8 @@ void FPerforceSourceControlModule::StartupModule()
 	PerforceSourceControlProvider.RegisterWorker( "Unshelve", FGetPerforceSourceControlWorker::CreateStatic( &CreateWorker<FPerforceUnshelveWorker>) );
 	PerforceSourceControlProvider.RegisterWorker( "DeleteShelved", FGetPerforceSourceControlWorker::CreateStatic( &CreateWorker<FPerforceDeleteShelveWorker>) );
 	PerforceSourceControlProvider.RegisterWorker( "DownloadFile", FGetPerforceSourceControlWorker::CreateStatic( &CreateWorker<FPerforceDownloadFileWorker>) );
+	PerforceSourceControlProvider.RegisterWorker( "CreateWorkspace", FGetPerforceSourceControlWorker::CreateStatic( &CreateWorker<FPerforceCreateWorkspaceWorker>) );
+	PerforceSourceControlProvider.RegisterWorker( "DeleteWorkspace", FGetPerforceSourceControlWorker::CreateStatic( &CreateWorker<FPerforceDeleteWorkspaceWorker>) );
 
 	// load our settings
 	PerforceSourceControlSettings.LoadSettings();
diff --git a/Engine/Plugins/Developer/PerforceSourceControl/Source/PerforceSourceControl/Private/PerforceSourceControlOperations.h b/Engine/Plugins/Developer/PerforceSourceControl/Source/PerforceSourceControl/Private/PerforceSourceControlOperations.h
--- a/Engine/Plugins/Developer/PerforceSourceControl/Source/PerforceSourceControl/Private/PerforceSourceControlOperations.h
+++ b/Engine/Plugins/Developer/PerforceSourceControl/Source/PerforceSourceControl/Private/PerforceSourceControlOperations.h
@@ -358,3 +358,25 @@ public:
 	virtual bool Execute(class FPerforceSourceControlCommand& InCommand) override;
 	virtual bool UpdateStates() const override;
 };
+
+class FPerforceCreateWorkspaceWorker : public IPerforceSourceControlWorker
+{
+public:
+	virtual ~FPerforceCreateWorkspaceWorker() = default;
+
+	// IPerforceSourceControlWorker interface
+	virtual FName GetName() const override;
+	virtual bool Execute(class FPerforceSourceControlCommand& InCommand) override;
+	virtual bool UpdateStates() const override;
+};
+
+class FPerforceDeleteWorkspaceWorker : public IPerforceSourceControlWorker
+{
+public:
+	virtual ~FPerforceDeleteWorkspaceWorker() = default;
+
+	// IPerforceSourceControlWorker interface
+	virtual FName GetName() const override;
+	virtual bool Execute(class FPerforceSourceControlCommand& InCommand) override;
+	virtual bool UpdateStates() const override;
+};
diff --git a/Engine/Plugins/Developer/PerforceSourceControl/Source/PerforceSourceControl/Private/PerforceSourceControlProvider.cpp b/Engine/Plugins/Developer/PerforceSourceControl/Source/PerforceSourceControl/Private/PerforceSourceControlProvider.cpp
--- a/Engine/Plugins/Developer/PerforceSourceControl/Source/PerforceSourceControl/Private/PerforceSourceControlProvider.cpp
+++ b/Engine/Plugins/Developer/PerforceSourceControl/Source/PerforceSourceControl/Private/PerforceSourceControlProvider.cpp
@@ -645,6 +645,69 @@ bool FPerforceSourceControlProvider::TryToDownloadFileFromBackgroundThread(const
 	return Command.bCommandSuccessful;
 }
 
+ECommandResult::Type FPerforceSourceControlProvider::SwitchWorkspace(FStringView NewWorkspaceName, FSourceControlResultInfo& OutResultInfo, FString* OutOldWorkspaceName)
+{
+	if (!CommandQueue.IsEmpty())
+	{
+		UE_LOG(LogSourceControl, Log, TEXT("Waiting on pending commands before switching workspace"));
+
+		// Run the busy loop while we wait for any current commands to be cleared.
+		while (!CommandQueue.IsEmpty())
+		{
+			// Tick the command queue and update progress.
+			Tick();
+
+			// Sleep for a bit so we don't busy-wait so much.
+			FPlatformProcess::Sleep(0.01f);
+		}
+	}
+
+	
+	Close();
+
+	// Do not call Init directly as we do not want to save the new workspace name to
+	// the source control settings!
+
+	FPerforceSourceControlModule& PerforceSourceControl = FModuleManager::LoadModuleChecked<FPerforceSourceControlModule>("PerforceSourceControl");
+	const FString OldWorkspaceName = PerforceSourceControl.AccessSettings().GetWorkspace();
+		
+	FPerforceSourceControlSettings& P4Settings = PerforceSourceControl.AccessSettings();
+
+	FString PortName = P4Settings.GetPort();
+	FString UserName = P4Settings.GetUserName();
+	FString ClientSpecName = FString(NewWorkspaceName);
+
+	if (FPerforceConnection::EnsureValidConnection(PortName, UserName, ClientSpecName, P4Settings.GetConnectionInfo()))
+	{
+		P4Settings.SetPort(PortName);
+		P4Settings.SetUserName(UserName);
+		P4Settings.SetWorkspace(ClientSpecName);
+		check(NewWorkspaceName == ClientSpecName);
+
+		bServerAvailable = true;
+
+		if (OutOldWorkspaceName != nullptr)
+		{
+			*OutOldWorkspaceName = OldWorkspaceName;
+		}
+
+		UE_LOG(LogSourceControl, Log, TEXT("Switched workspaces from '%s' to '%s%'"),  *OldWorkspaceName, *ClientSpecName);
+
+		return ECommandResult::Succeeded;
+	}
+	else
+	{
+		FText Message = FText::Format(	LOCTEXT("Perforce_ConnectionFailed", "Failed to re-establish the connection after switching to workspace {0}"),
+										FText::FromString(FString(NewWorkspaceName)));
+		OutResultInfo.ErrorMessages.Add(Message);
+
+		// The connection didn't work so we should try to restore the old workspace name
+		PerforceSourceControl.AccessSettings().SetWorkspace(OldWorkspaceName);
+		
+		return ECommandResult::Failed;
+	}
+}
+
 #if SOURCE_CONTROL_WITH_SLATE
 TSharedRef<class SWidget> FPerforceSourceControlProvider::MakeSettingsWidget() const
 {
diff --git a/Engine/Plugins/Developer/PerforceSourceControl/Source/PerforceSourceControl/Private/PerforceSourceControlProvider.h b/Engine/Plugins/Developer/PerforceSourceControl/Source/PerforceSourceControl/Private/PerforceSourceControlProvider.h
--- a/Engine/Plugins/Developer/PerforceSourceControl/Source/PerforceSourceControl/Private/PerforceSourceControlProvider.h
+++ b/Engine/Plugins/Developer/PerforceSourceControl/Source/PerforceSourceControl/Private/PerforceSourceControlProvider.h
@@ -51,6 +51,7 @@ public:
 	virtual TArray<FSourceControlChangelistRef> GetChangelists( EStateCacheUsage::Type InStateCacheUsage ) override;
 
 	virtual bool TryToDownloadFileFromBackgroundThread(const TSharedRef<class FDownloadFile>& InOperation, const TArray<FString>& InFiles) override;
+	virtual ECommandResult::Type SwitchWorkspace(FStringView NewWorkspaceName, FSourceControlResultInfo& OutResultInfo, FString* OutOldWorkspaceName) override;
 
 #if SOURCE_CONTROL_WITH_SLATE
 	virtual TSharedRef<class SWidget> MakeSettingsWidget() const override;
diff --git a/Engine/Source/Developer/SourceControl/Private/ISourceControlProvider.cpp b/Engine/Source/Developer/SourceControl/Private/ISourceControlProvider.cpp
--- a/Engine/Source/Developer/SourceControl/Private/ISourceControlProvider.cpp
+++ b/Engine/Source/Developer/SourceControl/Private/ISourceControlProvider.cpp
@@ -118,4 +118,15 @@ bool ISourceControlProvider::TryToDownloadFileFromBackgroundThread(const TShared
 	return false;
 }
 
+ECommandResult::Type ISourceControlProvider::SwitchWorkspace(FStringView NewWorkspaceName, FSourceControlResultInfo& OutResultInfo, FString* OutOldWorkspaceName)
+{
+	FFormatNamedArguments Arguments;
+	Arguments.Add(TEXT("ProviderName"), FText::FromName(GetName()));
+	FText Message = FText::Format(LOCTEXT("UnsupportedOperation", "SwitchWorkspace is not supported by source control provider '{ProviderName}'"), Arguments);
+
+	OutResultInfo.ErrorMessages.Add(Message);
+
+	return ECommandResult::Failed;
+}
+
 #undef LOCTEXT_NAMESPACE
diff --git a/Engine/Source/Developer/SourceControl/Private/SourceControlOperations.cpp b/Engine/Source/Developer/SourceControl/Private/SourceControlOperations.cpp
--- a/Engine/Source/Developer/SourceControl/Private/SourceControlOperations.cpp
+++ b/Engine/Source/Developer/SourceControl/Private/SourceControlOperations.cpp
@@ -4,6 +4,7 @@
 
 #include "Containers/StringView.h"
 #include "ISourceControlModule.h"
+#include "Misc/PathViews.h"
 #include "Misc/Paths.h"
 
 FDownloadFile::FDownloadFile(FStringView InTargetDirectory, EVerbosity InVerbosity)
@@ -23,3 +24,12 @@ FDownloadFile::FDownloadFile(FStringView InTargetDirectory, EVerbosity InVerbosi
 			*TargetDirectory, *Reason.ToString());
 	}
 }
+
+FCreateWorkspace::FCreateWorkspace(FStringView InWorkspaceName, FStringView InWorkspaceRoot)
+	: WorkspaceName(InWorkspaceName)
+{
+	TStringBuilder<512> AbsoluteWorkspaceName;
+	FPathViews::ToAbsolutePath(InWorkspaceRoot, AbsoluteWorkspaceName);
+
+	WorkspaceRoot = AbsoluteWorkspaceName.ToString();
+}
diff --git a/Engine/Source/Developer/SourceControl/Public/ISourceControlProvider.h b/Engine/Source/Developer/SourceControl/Public/ISourceControlProvider.h
--- a/Engine/Source/Developer/SourceControl/Public/ISourceControlProvider.h
+++ b/Engine/Source/Developer/SourceControl/Public/ISourceControlProvider.h
@@ -288,6 +288,18 @@ public:
 	 */
 	SOURCECONTROL_API virtual bool TryToDownloadFileFromBackgroundThread(const TSharedRef<class FDownloadFile>& InOperation, const TArray<FString>& InFiles);
 
+	/**
+	 * Used to switch the provider from one workspace to another. 
+	 * NOTE: This concept is currently only implemented for the perforce source control provider.
+	 * 
+	 * @param NewWorkspaceName			The name of the workspace to switch to
+	 * @param ResultInfo[out]			Errors and info messages generated will be written here
+	 * @param OutOldWorkspaceName[out]	The name of the previous workspace will be written to this FString if the pointer is valid (optional)
+	 * 
+	 * @return The result of the operation.
+	 */
+	SOURCECONTROL_API virtual ECommandResult::Type SwitchWorkspace(FStringView NewWorkspaceName, FSourceControlResultInfo& OutResultInfo, FString* OutOldWorkspaceName);
+
 	/**
 	 * Whether the provider uses local read-only state to signal whether a file is editable.
 	 */
diff --git a/Engine/Source/Developer/SourceControl/Public/SourceControlOperations.h b/Engine/Source/Developer/SourceControl/Public/SourceControlOperations.h
--- a/Engine/Source/Developer/SourceControl/Public/SourceControlOperations.h
+++ b/Engine/Source/Developer/SourceControl/Public/SourceControlOperations.h
@@ -762,4 +762,107 @@ private:
 	TMap<FString, FSharedBuffer> FileDataMap;
 };
 
+/** 
+ * Operation used to create a new workspace if the source control system supports this functionality.
+ */
+class FCreateWorkspace : public FSourceControlOperationBase
+{
+public:
+	using FClientViewMapping = TPair<FString,FString>;
+
+	FCreateWorkspace() = delete;
+
+	/**
+	 * Constructor
+	 * 
+	 * @param InWorkspaceName	The name of the workspace to create
+	 * @param InWorkspaceRoom	The file path to the workspace root (can be relative to the project)
+	 */
+	SOURCECONTROL_API FCreateWorkspace(FStringView InWorkspaceName, FStringView InWorkspaceRoot);
+	
+	/**
+	 * Add a new mapping for the client spec in the native format of the current source control provider.
+	 * These will be written to the client spec in the order that they are added.
+	 * 
+	 * @param DepotPath		The path in the source control depot to map from
+	 * @param ClientPath	The path on the local client to map too
+	 */
+	void AddNativeClientViewMapping(FStringView DepotPath, FStringView ClientPath)
+	{
+		ClientView.Emplace(DepotPath, ClientPath);
+	}
+
+	const FString& GetWorkspaceName() const
+	{
+		return WorkspaceName;
+	}
+
+	const FString& GetWorkspaceRoot() const
+	{
+		return WorkspaceRoot;
+	}
+
+	const TArray<FClientViewMapping> GetClientView() const
+	{
+		return ClientView;
+	}
+
+	// ISourceControlOperation interface
+
+	virtual FName GetName() const override
+	{
+		return "CreateWorkspace";
+	}
+
+	virtual FText GetInProgressString() const override
+	{
+		return LOCTEXT("SourceControl_CreateWorkspaceOperation", "Creating a workspace...");
+	}
+
+protected:
+
+	FString WorkspaceName;
+	FString WorkspaceRoot;
+	TArray<FClientViewMapping> ClientView;
+};
+
+/** Operation used to delete a workspace */
+class FDeleteWorkspace : public FSourceControlOperationBase
+{
+public:
+	FDeleteWorkspace() = delete;
+
+	/** 
+	 * Constructor
+	 * 
+	 * @param InWorkspaceName The name of the workspace to be deleted
+	 */
+	FDeleteWorkspace(FStringView InWorkspaceName)
+		: WorkspaceName(InWorkspaceName)
+	{
+
+	}
+
+	const FString& GetWorkspaceName() const
+	{
+		return WorkspaceName;
+	}
+
+	// ISourceControlOperation interface
+
+	virtual FName GetName() const override
+	{
+		return "DeleteWorkspace";
+	}
+
+	virtual FText GetInProgressString() const override
+	{
+		return LOCTEXT("SourceControl_DeleteWorkspaceOperation", "Deleting a workspace...");
+	}
+
+protected:
+
+	FString WorkspaceName;
+};
+
 #undef LOCTEXT_NAMESPACE
