CIS: Add file missing from CL 17350788

#rnx

diff --git a/Engine/Plugins/Developer/PerforceSourceControl/Source/PerforceSourceControl/Private/PerforceSourceControlOperations.cpp b/Engine/Plugins/Developer/PerforceSourceControl/Source/PerforceSourceControl/Private/PerforceSourceControlOperations.cpp
--- a/Engine/Plugins/Developer/PerforceSourceControl/Source/PerforceSourceControl/Private/PerforceSourceControlOperations.cpp
+++ b/Engine/Plugins/Developer/PerforceSourceControl/Source/PerforceSourceControl/Private/PerforceSourceControlOperations.cpp
@@ -2903,4 +2903,87 @@ bool FPerforceDownloadFileWorker::UpdateStates() const
 	return false;
 }
 
+FName FPerforceCreateWorkspaceWorker::GetName() const
+{
+	return "CreateWorkspace";
+}
+
+bool FPerforceCreateWorkspaceWorker::Execute(class FPerforceSourceControlCommand& InCommand)
+{
+	FScopedPerforceConnection ScopedConnection(InCommand);
+
+	if (!InCommand.IsCanceled() && ScopedConnection.IsValid())
+	{
+		FPerforceConnection& Connection = ScopedConnection.GetConnection();
+
+		TSharedRef<FCreateWorkspace, ESPMode::ThreadSafe> Operation = StaticCastSharedRef<FCreateWorkspace>(InCommand.Operation);
+
+		TStringBuilder<2048> ClientDesc;
+
+		ClientDesc << TEXT("Client:\t") << Operation->GetWorkspaceName() << TEXT("\n\n");
+		ClientDesc << TEXT("Owner:\t") << Connection.GetUser() << TEXT("\n\n");
+		ClientDesc << TEXT("Root:\t") << Operation->GetWorkspaceRoot() << TEXT("\n\n");
+		ClientDesc << TEXT("View:\t\n\n");
+
+		for (const FCreateWorkspace::FClientViewMapping& Mapping : Operation->GetClientView())
+		{
+			ClientDesc << TEXT("\t") << Mapping.Key;
+			ClientDesc << TEXT(" ");
+			ClientDesc << Mapping.Value << TEXT("\n");
+		}
+
+		InCommand.bCommandSuccessful = Connection.CreateWorkspace(	ClientDesc,
+																	FOnIsCancelled::CreateRaw(&InCommand, &FPerforceSourceControlCommand::IsCanceled),
+																	InCommand.ResultInfo.ErrorMessages);
+	}
+	
+	return InCommand.bCommandSuccessful;
+}
+
+bool FPerforceCreateWorkspaceWorker::UpdateStates() const
+{
+	return false;
+}
+
+FName FPerforceDeleteWorkspaceWorker::GetName() const
+{
+	return "DeleteWorkspace";
+}
+
+bool FPerforceDeleteWorkspaceWorker::Execute(class FPerforceSourceControlCommand& InCommand)
+{
+	FScopedPerforceConnection ScopedConnection(InCommand);
+
+	if (!InCommand.IsCanceled() && ScopedConnection.IsValid())
+	{
+		FPerforceConnection& Connection = ScopedConnection.GetConnection();
+
+		TSharedRef<FDeleteWorkspace, ESPMode::ThreadSafe> Operation = StaticCastSharedRef<FDeleteWorkspace>(InCommand.Operation);
+
+		TArray<FString> Parameters;
+		Parameters.Add(TEXT("-d")); // Do not print the header, we only want the actual file contents
+		Parameters.Add(Operation->GetWorkspaceName());
+
+		FP4RecordSet Records;
+
+		Connection.RunCommand(	TEXT("client"), 
+								Parameters, 
+								Records, 
+								InCommand.ResultInfo.ErrorMessages, 
+								FOnIsCancelled::CreateRaw(&InCommand, &FPerforceSourceControlCommand::IsCanceled), 
+								InCommand.bConnectionDropped);
+
+		// p4 client -d doesn't return any records, so the return value of RunCommand will always
+		// be false, so check the error messages instead.
+		InCommand.bCommandSuccessful = InCommand.ResultInfo.ErrorMessages.IsEmpty();
+	}
+
+	return InCommand.bCommandSuccessful;
+}
+
+bool FPerforceDeleteWorkspaceWorker::UpdateStates() const
+{
+	return false;
+}
+
 #undef LOCTEXT_NAMESPACE
