Only log why a potential workspace was rejected for consideration to when verbose logging for 'LogSourceControl is enabled

#rnx

- Due to mirage, the build machines can now end up initializing the perforce source control module and have access to many hundreds of potential workspaces. Logging each rejected workspace not only fills the log with spam, but can take several seconds (7 or more in tests) to complete, which is a waste of time.
- Changed the log item to verbose, rather than remove it so that people trying to debug a problem will still have access to the info if they need it.
- Technically for Mirage the build machines don't need a workspace as they will only ever be using p4 print, but our current architecture does not really allow for this (ideally Mirage would be creating/requesting it's own instance of the PerforceSourceControlProvider which we could make sure did not attempt to find a workspace) but changing the logging verbosity will serve for now.

diff --git a/Engine/Plugins/Developer/PerforceSourceControl/Source/PerforceSourceControl/Private/PerforceConnection.cpp b/Engine/Plugins/Developer/PerforceSourceControl/Source/PerforceSourceControl/Private/PerforceConnection.cpp
--- a/Engine/Plugins/Developer/PerforceSourceControl/Source/PerforceSourceControl/Private/PerforceConnection.cpp
+++ b/Engine/Plugins/Developer/PerforceSourceControl/Source/PerforceSourceControl/Private/PerforceConnection.cpp
@@ -715,7 +715,7 @@ bool FPerforceConnection::GetWorkspaceList(const FPerforceConnectionInfo& InConn
 					}
 					else
 					{
-						UE_LOG(LogSourceControl, Display, TEXT(" %s client specs rejected due to root directory mismatch (%s)"), *ClientName, *ClientRootPath);
+						UE_LOG(LogSourceControl, Verbose, TEXT(" %s client specs rejected due to root directory mismatch (%s)"), *ClientName, *ClientRootPath);
 					}
 
 					//Other useful fields: Description, Owner, Host
@@ -723,7 +723,7 @@ bool FPerforceConnection::GetWorkspaceList(const FPerforceConnectionInfo& InConn
 				}
 				else
 				{
-					UE_LOG(LogSourceControl, Display, TEXT(" %s client specs rejected due to host name mismatch (%s)"), *ClientName, *HostName);
+					UE_LOG(LogSourceControl, Verbose, TEXT(" %s client specs rejected due to host name mismatch (%s)"), *ClientName, *HostName);
 				}
 			}
 		}
