Fix abusive external dependencies validation error:
* Fixes Source Control cache corruption when submitting a changelist
* Updates external dependencies states when validating in case cached states are not valid anymore

#changelist validated

diff --git a/Engine/Plugins/Developer/PerforceSourceControl/Source/PerforceSourceControl/Private/PerforceSourceControlOperations.cpp b/Engine/Plugins/Developer/PerforceSourceControl/Source/PerforceSourceControl/Private/PerforceSourceControlOperations.cpp
--- a/Engine/Plugins/Developer/PerforceSourceControl/Source/PerforceSourceControl/Private/PerforceSourceControlOperations.cpp
+++ b/Engine/Plugins/Developer/PerforceSourceControl/Source/PerforceSourceControl/Private/PerforceSourceControlOperations.cpp
@@ -456,6 +456,8 @@ bool FPerforceCheckInWorker::Execute(FPerforceSourceControlCommand& InCommand)
 	{
 		FPerforceConnection& Connection = ScopedConnection.GetConnection();
 		FText CachedDescription;
+		FPerforceSourceControlModule& PerforceSourceControl = FModuleManager::GetModuleChecked<FPerforceSourceControlModule>("PerforceSourceControl");
+		FPerforceSourceControlProvider& Provider = PerforceSourceControl.GetProvider();
 
 		check(InCommand.Operation->GetName() == GetName());
 		TSharedRef<FCheckIn, ESPMode::ThreadSafe> Operation = StaticCastSharedRef<FCheckIn>(InCommand.Operation);
@@ -463,22 +465,22 @@ bool FPerforceCheckInWorker::Execute(FPerforceSourceControlCommand& InCommand)
 		TArray<FString> FilesToSubmit = InCommand.Files;
 
 		FPerforceSourceControlChangelist ChangeList(InCommand.Changelist);
+		FSourceControlChangelistStateRef ChangelistState = PerforceSourceControl.GetProvider().GetStateInternal(InCommand.Changelist);
 		TArray<FString> ReopenedFiles;
 
 		InCommand.bCommandSuccessful = true;
 
-		if (InCommand.Changelist.IsDefault())
+		// Adds changelist's files to FilesToSubmit so we can update their state after submit
+		if ((FilesToSubmit.Num() == 0) && (InCommand.Changelist.IsInitialized()))
 		{
-			// If the command has specified the default changelist but no files, then get all files from the default changelist
-			if (FilesToSubmit.Num() == 0 && InCommand.Changelist.IsInitialized())
+			Algo::Transform(ChangelistState->GetFilesStates(), FilesToSubmit, [](const FSourceControlStateRef& FileState)
 			{
-				FPerforceSourceControlModule& PerforceSourceControl = FPerforceSourceControlModule::Get();
-				TSharedRef<FPerforceSourceControlChangelistState, ESPMode::ThreadSafe> DefaultChangelistState = PerforceSourceControl.GetProvider().GetStateInternal(InCommand.Changelist);
-				Algo::Transform(DefaultChangelistState->Files, FilesToSubmit, [](const auto& FileState) {
-					return FileState->GetFilename();
-					});
-			}
+				return FileState->GetFilename();
+			});
+		}
 
+		if (InCommand.Changelist.IsDefault())
+		{
 			int32 NewChangeList = Connection.CreatePendingChangelist(Operation->GetDescription(), TArray<FString>(), FOnIsCancelled::CreateRaw(&InCommand, &FPerforceSourceControlCommand::IsCanceled), InCommand.ResultInfo.ErrorMessages);
 			if (NewChangeList > 0)
 			{
@@ -492,9 +494,6 @@ bool FPerforceCheckInWorker::Execute(FPerforceSourceControlCommand& InCommand)
 		}
 		else if (!Operation->GetDescription().IsEmpty())
 		{
-			FPerforceSourceControlModule& PerforceSourceControl = FPerforceSourceControlModule::Get();
-			FSourceControlChangelistStateRef ChangelistState = PerforceSourceControl.GetProvider().GetStateInternal(InCommand.Changelist);
-
 			// Retrieves cached description to restore it in case of submit failure.
 			CachedDescription = ChangelistState->GetDescriptionText();
 
@@ -521,9 +520,6 @@ bool FPerforceCheckInWorker::Execute(FPerforceSourceControlCommand& InCommand)
 			if (InCommand.bCommandSuccessful)
 			{
 				// Remove any deleted files from status cache
-				FPerforceSourceControlModule& PerforceSourceControl = FModuleManager::GetModuleChecked<FPerforceSourceControlModule>("PerforceSourceControl");
-				FPerforceSourceControlProvider& Provider = PerforceSourceControl.GetProvider();
-
 				TArray<TSharedRef<ISourceControlState, ESPMode::ThreadSafe>> States;
 				Provider.GetState(FilesToSubmit, States, EStateCacheUsage::Use);
 				for (const auto& State : States)
diff --git a/Engine/Plugins/Editor/DataValidation/Source/DataValidation/Private/DataValidationChangelist.cpp b/Engine/Plugins/Editor/DataValidation/Source/DataValidation/Private/DataValidationChangelist.cpp
--- a/Engine/Plugins/Editor/DataValidation/Source/DataValidation/Private/DataValidationChangelist.cpp
+++ b/Engine/Plugins/Editor/DataValidation/Source/DataValidation/Private/DataValidationChangelist.cpp
@@ -1,12 +1,15 @@
 // Copyright Epic Games, Inc. All Rights Reserved.
 
 #include "DataValidationChangelist.h"
+
+#include "Algo/AnyOf.h"
+#include "Algo/Transform.h"
 #include "AssetData.h"
 #include "AssetRegistryModule.h"
 #include "ISourceControlProvider.h"
 #include "ISourceControlModule.h"
 #include "SourceControlHelpers.h"
-#include "Algo/AnyOf.h"
+#include "SourceControlOperations.h"
 
 #define LOCTEXT_NAMESPACE "DataValidationChangelist"
 
@@ -110,9 +113,22 @@ EDataValidationResult UDataValidationChangelist::IsDataValid(TArray<FText>& Vali
 
 	bool bHasChangelistErrors = false;
 
-	for(FName& ExternalDependency : ExternalDependencies)
+	TArray<FString> ExternalDependenciesFilenames;
+
+	Algo::Transform(ExternalDependencies, ExternalDependenciesFilenames, [](const FName& InFilename) -> FString
+	{
+		return USourceControlHelpers::PackageFilename(InFilename.ToString());
+	});
+
+	// Update External dependencies state in case it changed from what is in cache
+	SourceControlProvider.Execute(ISourceControlOperation::Create<FUpdateStatus>(), ExternalDependenciesFilenames);
+
+	check(ExternalDependenciesFilenames.Num() == ExternalDependencies.Num());
+
+	for (int32 i = 0; i < ExternalDependenciesFilenames.Num(); ++i)
 	{
-		FString ExternalPackageFilename = USourceControlHelpers::PackageFilename(ExternalDependency.ToString());
+		const FString& ExternalPackageFilename = ExternalDependenciesFilenames[i];
+		const FName& ExternalDependency = ExternalDependencies[i];
 
 		FSourceControlStatePtr ExternalDependencyFileState = SourceControlProvider.GetState(ExternalPackageFilename, EStateCacheUsage::Use);
 
