Fixes hash based UpdateStatus unable to detect unchanged text files because of difference in P4 server line ending.
Replaced by forcing the diff command (hence avoiding errors for files not opened in workspace)

#changelist validated

diff --git a/Engine/Plugins/Developer/PerforceSourceControl/Source/PerforceSourceControl/Private/PerforceSourceControlOperations.cpp b/Engine/Plugins/Developer/PerforceSourceControl/Source/PerforceSourceControl/Private/PerforceSourceControlOperations.cpp
--- a/Engine/Plugins/Developer/PerforceSourceControl/Source/PerforceSourceControl/Private/PerforceSourceControlOperations.cpp
+++ b/Engine/Plugins/Developer/PerforceSourceControl/Source/PerforceSourceControl/Private/PerforceSourceControlOperations.cpp
@@ -1,7 +1,6 @@
 // Copyright Epic Games, Inc. All Rights Reserved.
 
 #include "PerforceSourceControlOperations.h"
-#include "Async/Async.h"
 #include "PerforceSourceControlPrivate.h"
 #include "HAL/FileManager.h"
 #include "Misc/Paths.h"
@@ -937,27 +936,7 @@ static void ParseBranchModificationResults(const FP4RecordSet& InRecords, const
 	}
 }
 
-struct FSourceControlDigest
-{
-	FString		FullPath;
-	FString		DepotDigest;
-	FMD5Hash	LocalDigest;
-};
-
-static FSourceControlDigest ComputeLocalDigest(FString InFilename, FString InDepotDigest)
-{
-	TRACE_CPUPROFILER_EVENT_SCOPE(FPerforceUpdateStatusWorker_ComputeLocalDigestTask);
-
-	FSourceControlDigest Result;
-
-	Result.FullPath = MoveTemp(InFilename);
-	Result.DepotDigest = MoveTemp(InDepotDigest);
-	Result.LocalDigest = FMD5Hash::HashFile(*Result.FullPath, nullptr);
-
-	return Result;
-}
-
-static void ParseUpdateStatusResults(const FP4RecordSet& InRecords, const TArray<FText>& ErrorMessages, TArray<FPerforceSourceControlState>& OutStates, TArray<FString>& OutModifiedFiles, const FString& ContentRoot, TMap<FString, FBranchModification>& BranchModifications)
+static void ParseUpdateStatusResults(const FP4RecordSet& InRecords, const TArray<FText>& ErrorMessages, TArray<FPerforceSourceControlState>& OutStates, const FString& ContentRoot, TMap<FString, FBranchModification>& BranchModifications)
 {
 	TRACE_CPUPROFILER_EVENT_SCOPE(FPerforceUpdateStatusWorker_ParseUpdateStatusResults);
 	// Build up a map of any other branch states
@@ -1023,8 +1002,6 @@ static void ParseUpdateStatusResults(const FP4RecordSet& InRecords, const TArray
 
 	}
 
-	TArray<TFuture<FSourceControlDigest>> SourceControlDigests;
-
 	// Iterate over each record found as a result of the command, parsing it for relevant information
 	for (int32 Index = 0; Index < InRecords.Num(); ++Index)
 	{
@@ -1058,11 +1035,6 @@ static void ParseUpdateStatusResults(const FP4RecordSet& InRecords, const TArray
 		FString Branch;
 		FString BranchFile;
 
-		if ((Digest.Len() > 0) && (FullPath.Len() > 0))
-		{
-			SourceControlDigests.Add(Async(EAsyncExecution::TaskGraph, [FullPath, Digest]() { return ComputeLocalDigest(FullPath, Digest); }));
-		}
-
 		if (DepotFileName.Split(ContentRoot, &Branch, &BranchFile))
 		{
 			// Sanitize
@@ -1227,23 +1199,6 @@ static void ParseUpdateStatusResults(const FP4RecordSet& InRecords, const TArray
 		}
 	}
 
-	for (const TFuture<FSourceControlDigest>& DigestResult : SourceControlDigests)
-	{
-		FSourceControlDigest Digest = DigestResult.Get();
-
-		if (!Digest.LocalDigest.IsValid())
-		{
-			continue;
-		}
-
-		FString LocalDigest = LexToString(Digest.LocalDigest);
-
-		if (!Digest.DepotDigest.Equals(LocalDigest, ESearchCase::IgnoreCase))
-		{
-			OutModifiedFiles.Add(Digest.FullPath);
-		}
-	}
-
 	// also see if we can glean anything from the error messages
 	for (int32 Index = 0; Index < ErrorMessages.Num(); ++Index)
 	{
@@ -1635,11 +1590,6 @@ bool FPerforceUpdateStatusWorker::Execute(FPerforceSourceControlCommand& InComma
 			// We want to include integration record information:
 			Parameters.Add(TEXT("-Or"));
 
-			if (Operation->ShouldUpdateModifiedState())
-			{
-				Parameters.Add(TEXT("-Olh"));
-			}
-
 			// Get the branches of interest for status updates
 			const FString& ContentRoot = InCommand.ContentRoot;
 			const TArray<FString>& StatusBranches = InCommand.StatusBranchNames;
@@ -1697,7 +1647,7 @@ bool FPerforceUpdateStatusWorker::Execute(FPerforceSourceControlCommand& InComma
 
 			FP4RecordSet Records;
 			InCommand.bCommandSuccessful &= Connection.RunCommand(TEXT("fstat"), Parameters, Records, InCommand.ResultInfo.ErrorMessages, FOnIsCancelled::CreateRaw(&InCommand, &FPerforceSourceControlCommand::IsCanceled), InCommand.bConnectionDropped);
-			ParseUpdateStatusResults(Records, InCommand.ResultInfo.ErrorMessages, OutStates, OutModifiedFiles,ContentRoot, BranchModifications);
+			ParseUpdateStatusResults(Records, InCommand.ResultInfo.ErrorMessages, OutStates, ContentRoot, BranchModifications);
 			RemoveRedundantErrors(InCommand, TEXT(" - no such file(s)."), false);
 			RemoveRedundantErrors(InCommand, TEXT("' is not under client's root '"));
 			RemoveRedundantErrors(InCommand, TEXT(" - protected namespace - access denied"), false);
@@ -1735,7 +1685,7 @@ bool FPerforceUpdateStatusWorker::Execute(FPerforceSourceControlCommand& InComma
 			TArray<FString> Parameters;
 			FP4RecordSet Records;
 			// Query for open files different than the versions stored in Perforce
-			Parameters.Add(TEXT("-sa"));
+			Parameters.Add(TEXT("-fsa"));
 			for (FString File : InCommand.Files)
 			{
 				if (IFileManager::Get().DirectoryExists(*File))
