Allow the type to be customized when creating a workspace.

#rnx

- Add a new type enum to FCreateWorkspace that can be used to decide on the type of client to create.
- Currently supporting types that map closely with perforce.

diff --git a/Engine/Plugins/Developer/PerforceSourceControl/Source/PerforceSourceControl/Private/PerforceSourceControlOperations.cpp b/Engine/Plugins/Developer/PerforceSourceControl/Source/PerforceSourceControl/Private/PerforceSourceControlOperations.cpp
--- a/Engine/Plugins/Developer/PerforceSourceControl/Source/PerforceSourceControl/Private/PerforceSourceControlOperations.cpp
+++ b/Engine/Plugins/Developer/PerforceSourceControl/Source/PerforceSourceControl/Private/PerforceSourceControlOperations.cpp
@@ -2878,6 +2878,8 @@ bool FPerforceCreateWorkspaceWorker::Execute(class FPerforceSourceControlCommand
 			ClientDesc << Mapping.Value << TEXT("\n");
 		}
 
+		AddType(*Operation, ClientDesc);
+		
 		InCommand.bCommandSuccessful = Connection.CreateWorkspace(	ClientDesc,
 																	FOnIsCancelled::CreateRaw(&InCommand, &FPerforceSourceControlCommand::IsCanceled),
 																	InCommand.ResultInfo.ErrorMessages);
@@ -2896,6 +2898,28 @@ FName FPerforceDeleteWorkspaceWorker::GetName() const
 	return "DeleteWorkspace";
 }
 
+void FPerforceCreateWorkspaceWorker::AddType(const FCreateWorkspace& Operation, FStringBuilderBase& ClientDesc)
+{
+	ClientDesc << TEXT("Type:\t");
+
+	switch (Operation.GetType())
+	{
+	case FCreateWorkspace::EType::Writeable:
+		ClientDesc << TEXT("writeable");
+		break;
+	case FCreateWorkspace::EType::ReadOnly:
+		ClientDesc << TEXT("readonly");
+		break;
+	case FCreateWorkspace::EType::Partitioned:
+		ClientDesc << TEXT("partitioned");
+		break;
+	default:
+		checkNoEntry();
+	}
+
+	ClientDesc << TEXT("\n");
+}
+
 bool FPerforceDeleteWorkspaceWorker::Execute(class FPerforceSourceControlCommand& InCommand)
 {
 	FScopedPerforceConnection ScopedConnection(InCommand);
diff --git a/Engine/Plugins/Developer/PerforceSourceControl/Source/PerforceSourceControl/Private/PerforceSourceControlOperations.h b/Engine/Plugins/Developer/PerforceSourceControl/Source/PerforceSourceControl/Private/PerforceSourceControlOperations.h
--- a/Engine/Plugins/Developer/PerforceSourceControl/Source/PerforceSourceControl/Private/PerforceSourceControlOperations.h
+++ b/Engine/Plugins/Developer/PerforceSourceControl/Source/PerforceSourceControl/Private/PerforceSourceControlOperations.h
@@ -368,6 +368,9 @@ public:
 	virtual FName GetName() const override;
 	virtual bool Execute(class FPerforceSourceControlCommand& InCommand) override;
 	virtual bool UpdateStates() const override;
+
+private:
+	void AddType(const class FCreateWorkspace& Operation, FStringBuilderBase& ClientDesc);
 };
 
 class FPerforceDeleteWorkspaceWorker : public IPerforceSourceControlWorker
diff --git a/Engine/Source/Developer/SourceControl/Public/SourceControlOperations.h b/Engine/Source/Developer/SourceControl/Public/SourceControlOperations.h
--- a/Engine/Source/Developer/SourceControl/Public/SourceControlOperations.h
+++ b/Engine/Source/Developer/SourceControl/Public/SourceControlOperations.h
@@ -768,6 +768,16 @@ private:
 class FCreateWorkspace : public FSourceControlOperationBase
 {
 public:
+	enum class EType
+	{
+		/* Default */
+		Writeable = 0,
+		/* Has a short life span and can only query/read/sync files */
+		ReadOnly,
+		/* Has a short life span but is allowed full functionality */
+		Partitioned
+	};
+
 	using FClientViewMapping = TPair<FString,FString>;
 
 	FCreateWorkspace() = delete;
@@ -792,6 +802,11 @@ public:
 		ClientView.Emplace(DepotPath, ClientPath);
 	}
 
+	void SetType(EType InType)
+	{
+		Type = InType;
+	}
+
 	const FString& GetWorkspaceName() const
 	{
 		return WorkspaceName;
@@ -807,6 +822,11 @@ public:
 		return ClientView;
 	}
 
+	EType GetType() const
+	{
+		return Type;
+	}
+
 	// ISourceControlOperation interface
 
 	virtual FName GetName() const override
@@ -823,6 +843,9 @@ protected:
 
 	FString WorkspaceName;
 	FString WorkspaceRoot;
+
+	EType Type = EType::Writeable;
+
 	TArray<FClientViewMapping> ClientView;
 };
 
