Fixes shelved files from another branch preventing FPerforceGetPendingChangelistsWorker to complete, hence Changelist Window to refresh.

#changelist validated

diff --git a/Engine/Plugins/Developer/PerforceSourceControl/Source/PerforceSourceControl/Private/PerforceSourceControlOperations.cpp b/Engine/Plugins/Developer/PerforceSourceControl/Source/PerforceSourceControl/Private/PerforceSourceControlOperations.cpp
--- a/Engine/Plugins/Developer/PerforceSourceControl/Source/PerforceSourceControl/Private/PerforceSourceControlOperations.cpp
+++ b/Engine/Plugins/Developer/PerforceSourceControl/Source/PerforceSourceControl/Private/PerforceSourceControlOperations.cpp
@@ -1926,8 +1926,7 @@ bool FPerforceGetPendingChangelistsWorker::Execute(FPerforceSourceControlCommand
 				Parameters.Add(ChangelistState.Changelist.ToString());
 
 				FP4RecordSet Records;
-				Connection.RunCommand(TEXT("describe"), Parameters, Records, InCommand.ResultInfo.ErrorMessages, FOnIsCancelled::CreateRaw(&InCommand, &FPerforceSourceControlCommand::IsCanceled), InCommand.bConnectionDropped);
-				InCommand.bCommandSuccessful &= (InCommand.ResultInfo.ErrorMessages.Num() == 0);
+				InCommand.bCommandSuccessful = Connection.RunCommand(TEXT("describe"), Parameters, Records, InCommand.ResultInfo.ErrorMessages, FOnIsCancelled::CreateRaw(&InCommand, &FPerforceSourceControlCommand::IsCanceled), InCommand.bConnectionDropped);
 
 				if (InCommand.bCommandSuccessful)
 				{
