Submitting new or previously deleted packages via the editor to perforce should no longer log errors about 'file(s) not on client'

#rnx

- In CL 18058385 -f (Force the diff) was applied to the diff command when trying to find unchanged files when submitting to perforce.
- When submitting files marked for add (either new files or files that are deleted at the head revision) using -f changed the message for the file from '- file(s) not opened for edit.' which we were removing from the error list to '- file(s) not on client.' which we were not.
- '- file(s) not on client.' has now been added to the list of errors we ignore from this specific 'p4 diff' command.

diff --git a/Engine/Plugins/Developer/PerforceSourceControl/Source/PerforceSourceControl/Private/PerforceSourceControlOperations.cpp b/Engine/Plugins/Developer/PerforceSourceControl/Source/PerforceSourceControl/Private/PerforceSourceControlOperations.cpp
--- a/Engine/Plugins/Developer/PerforceSourceControl/Source/PerforceSourceControl/Private/PerforceSourceControlOperations.cpp
+++ b/Engine/Plugins/Developer/PerforceSourceControl/Source/PerforceSourceControl/Private/PerforceSourceControlOperations.cpp
@@ -1704,6 +1704,7 @@ bool FPerforceUpdateStatusWorker::Execute(FPerforceSourceControlCommand& InComma
 			RemoveRedundantErrors(InCommand, TEXT(" - file(s) not opened for edit"));
 			RemoveRedundantErrors(InCommand, TEXT("' is not under client's root '"));
 			RemoveRedundantErrors(InCommand, TEXT(" - file(s) not opened on this client"));
+			RemoveRedundantErrors(InCommand, TEXT(" - file(s) not on client."));
 		}
 	}
 
