Adds "Delete new files on revert" option for Source Control
Prevents Uncontrolled Changelist's UpdateStatus usage to pollute the error logs.

#changelist validated

diff --git a/Engine/Plugins/Developer/PerforceSourceControl/Source/PerforceSourceControl/Private/PerforceSourceControlOperations.cpp b/Engine/Plugins/Developer/PerforceSourceControl/Source/PerforceSourceControl/Private/PerforceSourceControlOperations.cpp
--- a/Engine/Plugins/Developer/PerforceSourceControl/Source/PerforceSourceControl/Private/PerforceSourceControlOperations.cpp
+++ b/Engine/Plugins/Developer/PerforceSourceControl/Source/PerforceSourceControl/Private/PerforceSourceControlOperations.cpp
@@ -711,10 +711,16 @@ bool FPerforceRevertWorker::Execute(FPerforceSourceControlCommand& InCommand)
 			AppendChangelistParameter(Parameters);
 		}
 
-		if (StaticCastSharedRef<FRevert>(InCommand.Operation)->IsSoftRevert())
+		TSharedRef<FRevert> RevertOperation = StaticCastSharedRef<FRevert>(InCommand.Operation);
+
+		if (RevertOperation->IsSoftRevert())
 		{
 			Parameters.Add(TEXT("-k"));
 		}
+		else if (RevertOperation->ShouldDeleteNewFiles())
+		{
+			Parameters.Add(TEXT("-w"));
+		}
 
 		if (InCommand.Files.Num() != 0)
 		{
@@ -1657,10 +1663,6 @@ bool FPerforceUpdateStatusWorker::Execute(FPerforceSourceControlCommand& InComma
 			InCommand.bCommandSuccessful = true;
 		}
 
-		// update using any special hints passed in via the operation
-
-		bForceQuiet = Operation->ShouldBeQuiet();
-
 		if(Operation->ShouldUpdateHistory())
 		{
 			GetFileHistory(Connection, InCommand, InCommand.Files, OutStates, OutHistory);
@@ -1698,13 +1700,15 @@ bool FPerforceUpdateStatusWorker::Execute(FPerforceSourceControlCommand& InComma
 			}
 			InCommand.bCommandSuccessful &= Connection.RunCommand(TEXT("diff"), Parameters, Records, InCommand.ResultInfo.ErrorMessages, FOnIsCancelled::CreateRaw(&InCommand, &FPerforceSourceControlCommand::IsCanceled), InCommand.bConnectionDropped);
 
+			bool bMoveErrorsToInfo = !Operation->ShouldBeQuiet();
+
 			// Parse the results and store them in the command
 			ParseDiffResults(Records, OutModifiedFiles);
 			RemoveRedundantErrors(InCommand, TEXT(" - no such file(s)."));
 			RemoveRedundantErrors(InCommand, TEXT(" - file(s) not opened for edit"));
 			RemoveRedundantErrors(InCommand, TEXT("' is not under client's root '"));
 			RemoveRedundantErrors(InCommand, TEXT(" - file(s) not opened on this client"));
-			RemoveRedundantErrors(InCommand, TEXT(" - file(s) not on client."));
+			RemoveRedundantErrors(InCommand, TEXT("- file(s) not on client."), bMoveErrorsToInfo);
 		}
 	}
 
diff --git a/Engine/Source/Developer/SourceControl/Private/SourceControlPreferences.cpp b/Engine/Source/Developer/SourceControl/Private/SourceControlPreferences.cpp
new file mode 100644
--- /dev/null
+++ b/Engine/Source/Developer/SourceControl/Private/SourceControlPreferences.cpp
@@ -0,0 +1,13 @@
+// Copyright Epic Games, Inc. All Rights Reserved.
+
+#include "SourceControlPreferences.h"
+
+bool USourceControlPreferences::IsValidationTagEnabled()
+{
+	return GetDefault<USourceControlPreferences>()->bEnableValidationTag;
+}
+
+bool USourceControlPreferences::ShouldDeleteNewFilesOnRevert()
+{
+	return GetDefault<USourceControlPreferences>()->bShouldDeleteNewFilesOnRevert;
+}
diff --git a/Engine/Source/Developer/SourceControl/Public/SourceControlOperations.h b/Engine/Source/Developer/SourceControl/Public/SourceControlOperations.h
--- a/Engine/Source/Developer/SourceControl/Public/SourceControlOperations.h
+++ b/Engine/Source/Developer/SourceControl/Public/SourceControlOperations.h
@@ -7,6 +7,7 @@
 #include "Containers/StringFwd.h"
 #include "Memory/SharedBuffer.h"
 #include "ISourceControlChangelist.h"
+#include "SourceControlPreferences.h"
 #include "SourceControlOperationBase.h"
 
 #define LOCTEXT_NAMESPACE "SourceControl"
@@ -182,6 +183,11 @@ public:
 		return bIsSoftRevert;
 	}
 
+	bool ShouldDeleteNewFiles() const
+	{
+		return USourceControlPreferences::ShouldDeleteNewFilesOnRevert();
+	}
+
 protected:
 	bool	bIsSoftRevert = false;
 };
diff --git a/Engine/Source/Developer/SourceControl/Public/SourceControlPreferences.h b/Engine/Source/Developer/SourceControl/Public/SourceControlPreferences.h
new file mode 100644
--- /dev/null
+++ b/Engine/Source/Developer/SourceControl/Public/SourceControlPreferences.h
@@ -0,0 +1,32 @@
+// Copyright Epic Games, Inc. All Rights Reserved.
+
+#pragma once
+
+#include "CoreMinimal.h"
+#include "Engine/DeveloperSettings.h"
+
+#include "SourceControlPreferences.generated.h"
+
+
+/** Settings for the Source Control Integration */
+UCLASS(config = Editor, defaultconfig, meta = (DisplayName = "Source Control"))
+class SOURCECONTROL_API USourceControlPreferences : public UDeveloperSettings
+{
+	GENERATED_BODY()
+
+public:
+	/** Helper to access easily Enable Validation Tag setting */
+	static bool IsValidationTagEnabled();
+
+	/** Helper to access easily Should Delete New Files On Revert setting */
+	static bool ShouldDeleteNewFilesOnRevert();
+
+public:
+	/** If enabled, adds a tag in changelist descriptions when they are validated */
+	UPROPERTY(config, EditAnywhere, Category = SourceControl, meta = (ToolTip = "Adds validation tag to changelist description on submit."))
+	bool bEnableValidationTag = true;
+
+	/** If enabled, deletes new files when reverted. */
+	UPROPERTY(config, EditAnywhere, Category = SourceControl, meta = (ToolTip = "Deletes new files when reverted."))
+	bool bShouldDeleteNewFilesOnRevert = true;
+};
diff --git a/Engine/Source/Developer/SourceControl/SourceControl.Build.cs b/Engine/Source/Developer/SourceControl/SourceControl.Build.cs
--- a/Engine/Source/Developer/SourceControl/SourceControl.Build.cs
+++ b/Engine/Source/Developer/SourceControl/SourceControl.Build.cs
@@ -14,6 +14,7 @@ public class SourceControl : ModuleRules
 				"CoreUObject",
 				"SlateCore",
 				"InputCore",
+				"DeveloperSettings",
 			}
 		);
 
diff --git a/Engine/Source/Developer/UncontrolledChangelists/Private/UncontrolledChangelistState.cpp b/Engine/Source/Developer/UncontrolledChangelists/Private/UncontrolledChangelistState.cpp
--- a/Engine/Source/Developer/UncontrolledChangelists/Private/UncontrolledChangelistState.cpp
+++ b/Engine/Source/Developer/UncontrolledChangelists/Private/UncontrolledChangelistState.cpp
@@ -117,6 +117,7 @@ bool FUncontrolledChangelistState::AddFiles(const TArray<FString>& InFilenames,
 	{
 		auto UpdateStatusOperation = ISourceControlOperation::Create<FUpdateStatus>();
 		UpdateStatusOperation->SetUpdateModifiedState(true);
+		UpdateStatusOperation->SetQuiet(true);
 		SourceControlProvider.Execute(UpdateStatusOperation, InFilenames);
 	}
 
@@ -130,10 +131,12 @@ bool FUncontrolledChangelistState::AddFiles(const TArray<FString>& InFilenames,
 			const bool bFileExists = IFileManager::Get().FileExists(*FileState->GetFilename());
 
 			const bool bIsUncontrolled = (!bIsSourceControlled) && bFileExists;
-			const bool bIsDeleted = bIsSourceControlled && (!bFileExists);
+			// File doesn't exist and is not marked for delete
+			const bool bIsDeleted = bIsSourceControlled && (!bFileExists) && (!FileState->IsDeleted());
+			const bool bIsModified = FileState->IsModified() && (!FileState->IsDeleted());
 
 			const bool bIsCheckoutCompliant = (!bCheckCheckout) || (!FileState->IsCheckedOut());
-			const bool bIsStatusCompliant = (!bCheckStatus) || FileState->IsModified() || bIsUncontrolled || bIsDeleted;
+			const bool bIsStatusCompliant = (!bCheckStatus) || bIsModified || bIsUncontrolled || bIsDeleted;
 
 			if (bIsCheckoutCompliant && bIsStatusCompliant)
 			{
diff --git a/Engine/Source/Editor/SourceControlWindows/Private/SSourceControlChangelists.cpp b/Engine/Source/Editor/SourceControlWindows/Private/SSourceControlChangelists.cpp
--- a/Engine/Source/Editor/SourceControlWindows/Private/SSourceControlChangelists.cpp
+++ b/Engine/Source/Editor/SourceControlWindows/Private/SSourceControlChangelists.cpp
@@ -25,7 +25,7 @@
 #include "SSourceControlDescription.h"
 #include "SourceControlWindows.h"
 #include "SourceControlHelpers.h"
-#include "SourceControlSettings.h"
+#include "SourceControlPreferences.h"
 #include "AssetToolsModule.h"
 #include "ContentBrowserModule.h"
 #include "IContentBrowserSingleton.h"
@@ -1055,7 +1055,7 @@ void SSourceControlChangelistsWidget::OnSubmitChangelist()
 
 FText SSourceControlChangelistsWidget::UpdateChangelistDescriptionToSubmitIfNeeded(const bool bInValidationResult, const FText& InOriginalChangelistDescription) const
 {
-	if (bInValidationResult && USourceControlSettings::IsValidationTagEnabled() && (!HasValidationTag(InOriginalChangelistDescription)))
+	if (bInValidationResult && USourceControlPreferences::IsValidationTagEnabled() && (!HasValidationTag(InOriginalChangelistDescription)))
 	{
 		FStringOutputDevice Str;
 
diff --git a/Engine/Source/Editor/SourceControlWindows/Private/SourceControlSettings.cpp b/Engine/Source/Editor/SourceControlWindows/Private/SourceControlSettings.cpp
deleted file mode 100644
diff --git a/Engine/Source/Editor/SourceControlWindows/Public/SourceControlSettings.h b/Engine/Source/Editor/SourceControlWindows/Public/SourceControlSettings.h
deleted file mode 100644
diff --git a/Engine/Source/Editor/SourceControlWindows/SourceControlWindows.Build.cs b/Engine/Source/Editor/SourceControlWindows/SourceControlWindows.Build.cs
--- a/Engine/Source/Editor/SourceControlWindows/SourceControlWindows.Build.cs
+++ b/Engine/Source/Editor/SourceControlWindows/SourceControlWindows.Build.cs
@@ -24,7 +24,6 @@ public class SourceControlWindows : ModuleRules
 				"AssetTools",
 				"EditorFramework",
 				"WorkspaceMenuStructure",
-				"DeveloperSettings",
 				"UnrealEd"			// We need this dependency here because we use PackageTools.
 			}
 		);
