UE: add an option to use the P4Config
#Tests: PIE, switch between 2 branches with same config file, ensure that P4Config is correctly set

diff --git a/Engine/Plugins/Developer/PerforceSourceControl/Source/PerforceSourceControl/Private/PerforceConnectionInfo.h b/Engine/Plugins/Developer/PerforceSourceControl/Source/PerforceSourceControl/Private/PerforceConnectionInfo.h
--- a/Engine/Plugins/Developer/PerforceSourceControl/Source/PerforceSourceControl/Private/PerforceConnectionInfo.h
+++ b/Engine/Plugins/Developer/PerforceSourceControl/Source/PerforceSourceControl/Private/PerforceConnectionInfo.h
@@ -27,4 +27,7 @@ struct FPerforceConnectionInfo
 
 	/** Changelist number */
 	FString ChangelistNumber;
+
+	/** Use P4USER, P4PORT and P4CLIENT from P4 environment if set*/
+	bool bUseP4Config = false;
 };
diff --git a/Engine/Plugins/Developer/PerforceSourceControl/Source/PerforceSourceControl/Private/PerforceSourceControlSettings.cpp b/Engine/Plugins/Developer/PerforceSourceControl/Source/PerforceSourceControl/Private/PerforceSourceControlSettings.cpp
--- a/Engine/Plugins/Developer/PerforceSourceControl/Source/PerforceSourceControl/Private/PerforceSourceControlSettings.cpp
+++ b/Engine/Plugins/Developer/PerforceSourceControl/Source/PerforceSourceControl/Private/PerforceSourceControlSettings.cpp
@@ -1,6 +1,9 @@
 // Copyright Epic Games, Inc. All Rights Reserved.
 
 #include "PerforceSourceControlSettings.h"
+
+#include "PerforceSourceControlPrivate.h"
+
 #include "Misc/ScopeLock.h"
 #include "Misc/ConfigCacheIni.h"
 #include "Modules/ModuleManager.h"
@@ -77,18 +80,50 @@ void FPerforceSourceControlSettings::SetChangelistNumber(const FString& InString
 	ConnectionInfo.ChangelistNumber = InString;
 }
 
+bool FPerforceSourceControlSettings::GetUseP4Config() const
+{
+	FScopeLock ScopeLock(&CriticalSection);
+	return ConnectionInfo.bUseP4Config;
+}
+
+void FPerforceSourceControlSettings::SetUseP4Config(bool bInUseP4Config)
+{
+	FScopeLock ScopeLock(&CriticalSection);
+	ConnectionInfo.bUseP4Config = bInUseP4Config;
+	if (bInUseP4Config)
+	{
+		ImportP4Config();
+	}
+}
+
+void FPerforceSourceControlSettings::ImportP4Config()
+{
+	ClientApi TestP4;
+	ConnectionInfo.Port = ANSI_TO_TCHAR(TestP4.GetPort().Text());
+	ConnectionInfo.UserName = ANSI_TO_TCHAR(TestP4.GetUser().Text());
+	ConnectionInfo.Workspace = ANSI_TO_TCHAR(TestP4.GetClient().Text());
+}
+
 void FPerforceSourceControlSettings::LoadSettings()
 {
 	FScopeLock ScopeLock(&CriticalSection);
 	const FString& IniFile = SourceControlHelpers::GetSettingsIni();
 
-	if(!GConfig->GetString(*PerforceSettingsConstants::SettingsSection, TEXT("Port"), ConnectionInfo.Port, IniFile))
+	GConfig->GetBool(*PerforceSettingsConstants::SettingsSection, TEXT("UseP4Config"), ConnectionInfo.bUseP4Config, IniFile);
+	if (ConnectionInfo.bUseP4Config)
 	{
-		// backwards compatibility - previously we mis-specified the Port as 'Host'
-		GConfig->GetString(*PerforceSettingsConstants::SettingsSection, TEXT("Host"), ConnectionInfo.Port, IniFile);
+		ImportP4Config();
 	}
-	GConfig->GetString(*PerforceSettingsConstants::SettingsSection, TEXT("UserName"), ConnectionInfo.UserName, IniFile);
-	GConfig->GetString(*PerforceSettingsConstants::SettingsSection, TEXT("Workspace"), ConnectionInfo.Workspace, IniFile);
+	else
+	{
+		if(!GConfig->GetString(*PerforceSettingsConstants::SettingsSection, TEXT("Port"), ConnectionInfo.Port, IniFile))
+		{
+			// backwards compatibility - previously we mis-specified the Port as 'Host'
+			GConfig->GetString(*PerforceSettingsConstants::SettingsSection, TEXT("Host"), ConnectionInfo.Port, IniFile);
+		}
+		GConfig->GetString(*PerforceSettingsConstants::SettingsSection, TEXT("UserName"), ConnectionInfo.UserName, IniFile);
+		GConfig->GetString(*PerforceSettingsConstants::SettingsSection, TEXT("Workspace"), ConnectionInfo.Workspace, IniFile);
+	}	
 	GConfig->GetString(*PerforceSettingsConstants::SettingsSection, TEXT("HostOverride"), ConnectionInfo.HostOverride, IniFile);
 }
 
@@ -96,6 +131,8 @@ void FPerforceSourceControlSettings::SaveSettings() const
 {
 	FScopeLock ScopeLock(&CriticalSection);
 	const FString& IniFile = SourceControlHelpers::GetSettingsIni();
+	
+	GConfig->SetBool(*PerforceSettingsConstants::SettingsSection, TEXT("UseP4Config"), ConnectionInfo.bUseP4Config, IniFile);
 	GConfig->SetString(*PerforceSettingsConstants::SettingsSection, TEXT("Port"), *ConnectionInfo.Port, IniFile);
 	GConfig->SetString(*PerforceSettingsConstants::SettingsSection, TEXT("UserName"), *ConnectionInfo.UserName, IniFile);
 	GConfig->SetString(*PerforceSettingsConstants::SettingsSection, TEXT("Workspace"), *ConnectionInfo.Workspace, IniFile);
@@ -111,7 +148,7 @@ FPerforceConnectionInfo FPerforceSourceControlSettings::GetConnectionInfo() cons
 	}
 	
 	// The password needs to be gotten straight from the input UI, its not stored anywhere else
-	FString Password = SPerforceSourceControlSettings::GetPassword();
+	const FString Password = SPerforceSourceControlSettings::GetPassword();
 	if(!Password.IsEmpty())
 	{
 		OutConnectionInfo.Password = Password;
@@ -119,7 +156,7 @@ FPerforceConnectionInfo FPerforceSourceControlSettings::GetConnectionInfo() cons
 
 	// Ticket is stored in the provider (this is only set by the command line so should be safe to access without threading protection)
 	FPerforceSourceControlModule& PerforceSourceControl = FModuleManager::GetModuleChecked<FPerforceSourceControlModule>( "PerforceSourceControl" );
-	FPerforceSourceControlProvider& Provider = PerforceSourceControl.GetProvider();
+	const FPerforceSourceControlProvider& Provider = PerforceSourceControl.GetProvider();
 	OutConnectionInfo.Ticket = Provider.GetTicket();
 
 	return OutConnectionInfo;
diff --git a/Engine/Plugins/Developer/PerforceSourceControl/Source/PerforceSourceControl/Private/PerforceSourceControlSettings.h b/Engine/Plugins/Developer/PerforceSourceControl/Source/PerforceSourceControl/Private/PerforceSourceControlSettings.h
--- a/Engine/Plugins/Developer/PerforceSourceControl/Source/PerforceSourceControl/Private/PerforceSourceControlSettings.h
+++ b/Engine/Plugins/Developer/PerforceSourceControl/Source/PerforceSourceControl/Private/PerforceSourceControlSettings.h
@@ -37,6 +37,11 @@ public:
 
 	void SetChangelistNumber(const FString& InString);
 
+	/* Import the P4USER, P4PORT, P4CLIENT from the P4 environment if set. */
+	bool GetUseP4Config() const;
+
+	void SetUseP4Config(bool bInUseP4Config);
+
 	/** Load settings from ini file */
 	void LoadSettings();
 
@@ -47,6 +52,9 @@ public:
 	FPerforceConnectionInfo GetConnectionInfo() const;
 
 private:
+
+	void ImportP4Config();
+	
 	/** A critical section for settings access */
 	mutable FCriticalSection CriticalSection;
 
diff --git a/Engine/Plugins/Developer/PerforceSourceControl/Source/PerforceSourceControl/Private/SPerforceSourceControlSettings.cpp b/Engine/Plugins/Developer/PerforceSourceControl/Source/PerforceSourceControl/Private/SPerforceSourceControlSettings.cpp
--- a/Engine/Plugins/Developer/PerforceSourceControl/Source/PerforceSourceControl/Private/SPerforceSourceControlSettings.cpp
+++ b/Engine/Plugins/Developer/PerforceSourceControl/Source/PerforceSourceControl/Private/SPerforceSourceControlSettings.cpp
@@ -4,9 +4,9 @@
 #include "PerforceSourceControlPrivate.h"
 #include "Widgets/Views/STableRow.h"
 #include "Modules/ModuleManager.h"
-#include "Widgets/Images/SImage.h"
 #include "Widgets/Input/SButton.h"
 #include "Widgets/Views/SListView.h"
+#include "Widgets/Input/SCheckBox.h"
 #include "Widgets/Input/SComboBox.h"
 #include "EditorStyleSet.h"
 #include "ISourceControlModule.h"
@@ -61,6 +61,15 @@ void SPerforceSourceControlSettings::Construct(const FArguments& InArgs)
 			.Padding(FMargin(0.0f, 0.0f, 0.0f, 10.0f))
 			.VAlign(VAlign_Center)
 			.HAlign(HAlign_Right)
+			[
+				SNew(STextBlock)
+				.Text(LOCTEXT("UseP4ConfigLabel", "Use P4 Config"))
+				.ToolTipText( LOCTEXT("UseP4Config_Tooltip", "Read the P4USER, P4PORT, P4CLIENT from Perforce environment variables and P4CONFIG file rather than the UE project settings files") )
+			]
+			+SVerticalBox::Slot()
+			.Padding(FMargin(0.0f, 0.0f, 0.0f, 10.0f))
+			.VAlign(VAlign_Center)
+			.HAlign(HAlign_Right)
 			[
 				SNew(STextBlock)
 				.Text(LOCTEXT("PortLabel", "Server"))
@@ -121,12 +130,21 @@ void SPerforceSourceControlSettings::Construct(const FArguments& InArgs)
 			+SVerticalBox::Slot()
 			.VAlign(VAlign_Center)
 			.Padding(FMargin(0.0f, 0.0f, 0.0f, 10.0f))
+			[
+				SNew(SCheckBox)
+				.IsChecked(this, &SPerforceSourceControlSettings::IsP4ConfigChecked)
+				.OnCheckStateChanged(this, &SPerforceSourceControlSettings::OnP4ConfigCheckStatusChanged)
+			]
+			+SVerticalBox::Slot()
+			.VAlign(VAlign_Center)
+			.Padding(FMargin(0.0f, 0.0f, 0.0f, 10.0f))
 			[
 				SNew(SEditableTextBox)
 				.Text(this, &SPerforceSourceControlSettings::GetPortText)
 				.ToolTipText( LOCTEXT("PortLabel_Tooltip", "The server and port for your Perforce server. Usage ServerName:1234.") )
-				.OnTextCommitted(this, &SPerforceSourceControlSettings::OnPortTextCommited)
-				.OnTextChanged(this, &SPerforceSourceControlSettings::OnPortTextCommited, ETextCommit::Default)
+				.OnTextCommitted(this, &SPerforceSourceControlSettings::OnPortTextCommitted)
+				.OnTextChanged(this, &SPerforceSourceControlSettings::OnPortTextCommitted, ETextCommit::Default)
+				.IsEnabled(this, &SPerforceSourceControlSettings::IsP4ConfigDisabled)
 			]
 			+SVerticalBox::Slot()
 			.Padding(FMargin(0.0f, 0.0f, 0.0f, 10.0f))
@@ -135,8 +153,9 @@ void SPerforceSourceControlSettings::Construct(const FArguments& InArgs)
 				SNew(SEditableTextBox)
 				.Text(this, &SPerforceSourceControlSettings::GetUserNameText)
 				.ToolTipText( LOCTEXT("UserNameLabel_Tooltip", "Perforce username.") )
-				.OnTextCommitted(this, &SPerforceSourceControlSettings::OnUserNameTextCommited)
-				.OnTextChanged(this, &SPerforceSourceControlSettings::OnUserNameTextCommited, ETextCommit::Default)
+				.OnTextCommitted(this, &SPerforceSourceControlSettings::OnUserNameTextCommitted)
+				.OnTextChanged(this, &SPerforceSourceControlSettings::OnUserNameTextCommitted, ETextCommit::Default)
+				.IsEnabled(this, &SPerforceSourceControlSettings::IsP4ConfigDisabled)
 			]
 			+SVerticalBox::Slot()
 			.Padding(FMargin(0.0f, 0.0f, 0.0f, 10.0f))
@@ -145,8 +164,9 @@ void SPerforceSourceControlSettings::Construct(const FArguments& InArgs)
 				SNew(SEditableTextBox)
 				.Text(this, &SPerforceSourceControlSettings::GetWorkspaceText)
 				.ToolTipText( LOCTEXT("WorkspaceLabel_Tooltip", "Perforce workspace.") )
-				.OnTextCommitted(this, &SPerforceSourceControlSettings::OnWorkspaceTextCommited)
-				.OnTextChanged(this, &SPerforceSourceControlSettings::OnWorkspaceTextCommited, ETextCommit::Default)
+				.OnTextCommitted(this, &SPerforceSourceControlSettings::OnWorkspaceTextCommitted)
+				.OnTextChanged(this, &SPerforceSourceControlSettings::OnWorkspaceTextCommitted, ETextCommit::Default)
+				.IsEnabled(this, &SPerforceSourceControlSettings::IsP4ConfigDisabled)
 			]
 			+SVerticalBox::Slot()
 			.Padding(FMargin(0.0f, 0.0f, 0.0f, 10.0f))
@@ -156,6 +176,7 @@ void SPerforceSourceControlSettings::Construct(const FArguments& InArgs)
 				.OnGetMenuContent(this, &SPerforceSourceControlSettings::OnGetMenuContent)
 				.ContentPadding(1)
 				.ToolTipText( LOCTEXT("AutoWorkspaces_Tooltip", "Choose from a list of available workspaces. Requires a server and username before use.") )
+				.IsEnabled(this, &SPerforceSourceControlSettings::IsP4ConfigDisabled)
 				.ButtonContent()
 				[
 					SNew( STextBlock )
@@ -170,8 +191,8 @@ void SPerforceSourceControlSettings::Construct(const FArguments& InArgs)
 				.Visibility(bAllowP4NonTicketBasedLogins ? EVisibility::Visible : EVisibility::Collapsed)
 				.Text(this, &SPerforceSourceControlSettings::GetHostText)
 				.ToolTipText(LOCTEXT("HostLabel_Tooltip", "If you wish to impersonate a particular host, enter this here. This is not normally needed."))
-				.OnTextCommitted(this, &SPerforceSourceControlSettings::OnHostTextCommited)
-				.OnTextChanged(this, &SPerforceSourceControlSettings::OnHostTextCommited, ETextCommit::Default)
+				.OnTextCommitted(this, &SPerforceSourceControlSettings::OnHostTextCommitted)
+				.OnTextChanged(this, &SPerforceSourceControlSettings::OnHostTextCommitted, ETextCommit::Default)
 			]
 			+ SVerticalBox::Slot()
 			.Padding(FMargin(0.0f, 0.0f, 0.0f, 10.0f))
@@ -209,13 +230,32 @@ void SPerforceSourceControlSettings::Tick( const FGeometry& AllottedGeometry, co
 	return SCompoundWidget::Tick(AllottedGeometry, InCurrentTime, InDeltaTime );
 }
 
+ECheckBoxState SPerforceSourceControlSettings::IsP4ConfigChecked() const
+{
+	FPerforceSourceControlModule& PerforceSourceControl = FModuleManager::LoadModuleChecked<FPerforceSourceControlModule>( "PerforceSourceControl" );
+	return PerforceSourceControl.AccessSettings().GetUseP4Config() ? ECheckBoxState::Checked : ECheckBoxState::Unchecked;
+}
+
+void SPerforceSourceControlSettings::OnP4ConfigCheckStatusChanged(ECheckBoxState NewState) const
+{
+	FPerforceSourceControlModule& PerforceSourceControl = FModuleManager::LoadModuleChecked<FPerforceSourceControlModule>( "PerforceSourceControl" );
+	PerforceSourceControl.AccessSettings().SetUseP4Config(NewState == ECheckBoxState::Checked);
+	PerforceSourceControl.SaveSettings();
+}
+
+bool SPerforceSourceControlSettings::IsP4ConfigDisabled() const
+{
+	FPerforceSourceControlModule& PerforceSourceControl = FModuleManager::LoadModuleChecked<FPerforceSourceControlModule>( "PerforceSourceControl" );
+	return !PerforceSourceControl.AccessSettings().GetUseP4Config();
+}
+
 FText SPerforceSourceControlSettings::GetPortText() const
 {
 	FPerforceSourceControlModule& PerforceSourceControl = FModuleManager::LoadModuleChecked<FPerforceSourceControlModule>( "PerforceSourceControl" );
 	return FText::FromString(PerforceSourceControl.AccessSettings().GetPort());
 }
 
-void SPerforceSourceControlSettings::OnPortTextCommited(const FText& InText, ETextCommit::Type InCommitType) const
+void SPerforceSourceControlSettings::OnPortTextCommitted(const FText& InText, ETextCommit::Type InCommitType) const
 {
 	FPerforceSourceControlModule& PerforceSourceControl = FModuleManager::LoadModuleChecked<FPerforceSourceControlModule>( "PerforceSourceControl" );
 	PerforceSourceControl.AccessSettings().SetPort(InText.ToString());
@@ -228,7 +268,7 @@ FText SPerforceSourceControlSettings::GetUserNameText() const
 	return FText::FromString(PerforceSourceControl.AccessSettings().GetUserName());
 }
 
-void SPerforceSourceControlSettings::OnUserNameTextCommited(const FText& InText, ETextCommit::Type InCommitType) const
+void SPerforceSourceControlSettings::OnUserNameTextCommitted(const FText& InText, ETextCommit::Type InCommitType) const
 {
 	FPerforceSourceControlModule& PerforceSourceControl = FModuleManager::LoadModuleChecked<FPerforceSourceControlModule>( "PerforceSourceControl" );
 	PerforceSourceControl.AccessSettings().SetUserName(InText.ToString());
@@ -241,7 +281,7 @@ FText SPerforceSourceControlSettings::GetWorkspaceText() const
 	return FText::FromString(PerforceSourceControl.AccessSettings().GetWorkspace());
 }
 
-void SPerforceSourceControlSettings::OnWorkspaceTextCommited(const FText& InText, ETextCommit::Type InCommitType) const
+void SPerforceSourceControlSettings::OnWorkspaceTextCommitted(const FText& InText, ETextCommit::Type InCommitType) const
 {
 	FPerforceSourceControlModule& PerforceSourceControl = FModuleManager::LoadModuleChecked<FPerforceSourceControlModule>( "PerforceSourceControl" );
 	PerforceSourceControl.AccessSettings().SetWorkspace(InText.ToString());
@@ -254,7 +294,7 @@ FText SPerforceSourceControlSettings::GetHostText() const
 	return FText::FromString(PerforceSourceControl.AccessSettings().GetHostOverride());
 }
 
-void SPerforceSourceControlSettings::OnHostTextCommited(const FText& InText, ETextCommit::Type InCommitType) const
+void SPerforceSourceControlSettings::OnHostTextCommitted(const FText& InText, ETextCommit::Type InCommitType) const
 {
 	FPerforceSourceControlModule& PerforceSourceControl = FModuleManager::LoadModuleChecked<FPerforceSourceControlModule>( "PerforceSourceControl" );
 	PerforceSourceControl.AccessSettings().SetHostOverride(InText.ToString());
@@ -269,7 +309,7 @@ void SPerforceSourceControlSettings::QueryWorkspaces()
 		CurrentWorkspace = FString();
 
 		// fire off the workspace query
-		ISourceControlModule& SourceControl = FModuleManager::LoadModuleChecked<ISourceControlModule>( "SourceControl" );
+		const ISourceControlModule& SourceControl = FModuleManager::LoadModuleChecked<ISourceControlModule>( "SourceControl" );
 		ISourceControlProvider& Provider = SourceControl.GetProvider();
 		GetWorkspacesOperation = ISourceControlOperation::Create<FGetWorkspaces>();
 		Provider.Execute(GetWorkspacesOperation.ToSharedRef(), EConcurrency::Asynchronous, FSourceControlOperationComplete::CreateSP(this, &SPerforceSourceControlSettings::OnSourceControlOperationComplete) );
@@ -373,7 +413,7 @@ EVisibility SPerforceSourceControlSettings::GetWorkspaceListVisibility() const
 	return State == ESourceControlOperationState::Queried && Workspaces.Num() > 0 ? EVisibility::Visible : EVisibility::Collapsed;
 }
 
-TSharedRef<ITableRow> SPerforceSourceControlSettings::OnGenerateWorkspaceRow(TSharedRef<FString> InItem, const TSharedRef<STableViewBase>& OwnerTable)
+TSharedRef<ITableRow> SPerforceSourceControlSettings::OnGenerateWorkspaceRow(TSharedRef<FString> InItem, const TSharedRef<STableViewBase>& OwnerTable) const
 {
 	return
 		SNew(SComboRow< TSharedRef<FString> >, OwnerTable)
@@ -404,11 +444,11 @@ FText SPerforceSourceControlSettings::OnGetButtonText() const
 	return FText::FromString(CurrentWorkspace);
 }
 
-FReply SPerforceSourceControlSettings::OnCancelWorkspacesRequest()
+FReply SPerforceSourceControlSettings::OnCancelWorkspacesRequest() const
 {
 	if(GetWorkspacesOperation.IsValid())
 	{
-		ISourceControlModule& SourceControl = FModuleManager::LoadModuleChecked<ISourceControlModule>( "SourceControl" );
+		const ISourceControlModule& SourceControl = FModuleManager::LoadModuleChecked<ISourceControlModule>( "SourceControl" );
 		SourceControl.GetProvider().CancelOperation(GetWorkspacesOperation.ToSharedRef());
 	}
 	return FReply::Handled();
diff --git a/Engine/Plugins/Developer/PerforceSourceControl/Source/PerforceSourceControl/Private/SPerforceSourceControlSettings.h b/Engine/Plugins/Developer/PerforceSourceControl/Source/PerforceSourceControl/Private/SPerforceSourceControlSettings.h
--- a/Engine/Plugins/Developer/PerforceSourceControl/Source/PerforceSourceControl/Private/SPerforceSourceControlSettings.h
+++ b/Engine/Plugins/Developer/PerforceSourceControl/Source/PerforceSourceControl/Private/SPerforceSourceControlSettings.h
@@ -65,29 +65,38 @@ private:
 	/** Fire off a source control operation to see what workspaces we have */
 	void QueryWorkspaces();
 
+	/** Delegate to check if we should use P4 Config from settings */
+	ECheckBoxState IsP4ConfigChecked() const;
+
+	/** Delegate to enable or disable P4 Config */
+	void OnP4ConfigCheckStatusChanged(ECheckBoxState NewState) const;
+
+	/** Delegate to enable Port, Username and input when P4 Config is not used */
+	bool IsP4ConfigDisabled() const;
+
 	/** Delegate to get port text from settings */
 	FText GetPortText() const;
 
 	/** Delegate to commit port text to settings */
-	void OnPortTextCommited(const FText& InText, ETextCommit::Type InCommitType) const;
+	void OnPortTextCommitted(const FText& InText, ETextCommit::Type InCommitType) const;
 
 	/** Delegate to get user name text from settings */
 	FText GetUserNameText() const;
 
 	/** Delegate to commit user name text to settings */
-	void OnUserNameTextCommited(const FText& InText, ETextCommit::Type InCommitType) const;
+	void OnUserNameTextCommitted(const FText& InText, ETextCommit::Type InCommitType) const;
 
 	/** Delegate to get workspace text from settings */
 	FText GetWorkspaceText() const;
 
 	/** Delegate to commit workspace text to settings */
-	void OnWorkspaceTextCommited(const FText& InText, ETextCommit::Type InCommitType) const;
+	void OnWorkspaceTextCommitted(const FText& InText, ETextCommit::Type InCommitType) const;
 
 	/** Delegate to get host text from settings */
 	FText GetHostText() const;
 
 	/** Delegate to commit host text to settings */
-	void OnHostTextCommited(const FText& InText, ETextCommit::Type InCommitType) const;
+	void OnHostTextCommitted(const FText& InText, ETextCommit::Type InCommitType) const;
 
 	/** Delegate called when a source control operation has completed */
 	void OnSourceControlOperationComplete(const FSourceControlOperationRef& InOperation, ECommandResult::Type InResult);
@@ -105,7 +114,7 @@ private:
 	EVisibility GetWorkspaceListVisibility() const;
 
 	/** Delegate to generate a row in the workspace list */
-	TSharedRef<class ITableRow> OnGenerateWorkspaceRow(TSharedRef<FString> InItem, const TSharedRef<STableViewBase>& OwnerTable);
+	TSharedRef<class ITableRow> OnGenerateWorkspaceRow(TSharedRef<FString> InItem, const TSharedRef<STableViewBase>& OwnerTable) const;
 
 	/** Delegate for when a workspace item is clicked */
 	void OnWorkspaceSelected(TSharedPtr<FString> InItem, ESelectInfo::Type InSelectInfo);
@@ -114,7 +123,7 @@ private:
 	FText OnGetButtonText() const;
 
 	/** Delegate that cancels the workspace request in progress */
-	FReply OnCancelWorkspacesRequest();
+	FReply OnCancelWorkspacesRequest() const;
 
 	/** Get the image to display on the advanced pulldown */
 	const FSlateBrush* GetAdvancedPulldownImage() const;
