Fix a mismerge in the FCreateWorkspace command introduced in CL 18946159 that prevented the FCreateWorkspace::EType value from being used.

#rnx

diff --git a/Engine/Plugins/Developer/PerforceSourceControl/Source/PerforceSourceControl/Private/PerforceSourceControlOperations.cpp b/Engine/Plugins/Developer/PerforceSourceControl/Source/PerforceSourceControl/Private/PerforceSourceControlOperations.cpp
--- a/Engine/Plugins/Developer/PerforceSourceControl/Source/PerforceSourceControl/Private/PerforceSourceControlOperations.cpp
+++ b/Engine/Plugins/Developer/PerforceSourceControl/Source/PerforceSourceControl/Private/PerforceSourceControlOperations.cpp
@@ -2947,6 +2947,8 @@ bool FPerforceCreateWorkspaceWorker::Execute(class FPerforceSourceControlCommand
 			ClientDesc << Mapping.Value << TEXT("\n");
 		}
 
+		AddType(*Operation, ClientDesc);
+
 		InCommand.bCommandSuccessful = Connection.CreateWorkspace(	ClientDesc,
 																	FOnIsCancelled::CreateRaw(&InCommand, &FPerforceSourceControlCommand::IsCanceled),
 																	InCommand.ResultInfo.ErrorMessages);
@@ -2955,6 +2957,29 @@ bool FPerforceCreateWorkspaceWorker::Execute(class FPerforceSourceControlCommand
 	return InCommand.bCommandSuccessful;
 }
 
+void FPerforceCreateWorkspaceWorker::AddType(const FCreateWorkspace& Operation, FStringBuilderBase& ClientDesc)
+{
+	ClientDesc << TEXT("Type:\t");
+
+	switch (Operation.GetType())
+	{
+		case FCreateWorkspace::EType::Writeable:
+			ClientDesc << TEXT("writeable");
+			break;
+		case FCreateWorkspace::EType::ReadOnly:
+			ClientDesc << TEXT("readonly");
+			break;
+		case FCreateWorkspace::EType::Partitioned:
+			ClientDesc << TEXT("partitioned");
+			break;
+		default:
+			checkNoEntry();
+	}
+
+	ClientDesc << TEXT("\n");
+}
+
+
 bool FPerforceCreateWorkspaceWorker::UpdateStates() const
 {
 	return false;
