Allow additional source control providers to be created so that specialized subsystems can maintain their own separate connection without interrupting the normal editor workflows.

#rnx

### Misc
- FPerforceSourceControlRevision, FPerforceConnection and FPerforceSourceControlLabel now own a reference to their FPerforceSourceControlProvider and accesses the connection settings from that rather than via the module singleton.
- FPerforceSourceControlCommand now accesses the connection settings by finding the FPerforceSourceControlProvider that the given IPerforceSourceControlWorker references.
- FPerforceSourceControlProvider now owns it's own FPerforceSourceControlSettings rather than accessing a shared one owned by the module.
- Marked a number of derived classes as final.

### FPerforceConnection
- ::EnsureValidConnection can now work without a valid client (although most commands will require a client to run, things like p4 print does not).
-- Due to this change we might not catch cases where the the system cannot connect to source control as the user is not logged in as this would get caught when we checked the workspace by calling ::TestConnection.
-- To get around this ::EnsureValidConnection will now call "p4 login -s" via the newly added TestLoginConnection.
- Renamed ::TestConnection to ::TestClientConnection as we are really testing the connection to the specified workspace.
- There does seem to be some inconsistency in the logging, depending how a perforce source control provider first creates its connection as in some cases FPerforceConnection::EnsureValidConnection gets called first and in other cases FPerforceConnection::EstablishConnection is called. We should unify these paths but this change list is already so large I am opting to postpone this to it's own work item.

### FPerforceSourceControlModule
- Removed all of the singleton style accessors so that classes inside of the module can no longer assume that there is only one FPerforceSourceControlProvider. Updated the documentation to match.
- Since there is no reason for classes in the module to need to access FPerforceSourceControlModule, ::Get has been removed
- The responsibility to register all of the worker creation delegates has been moved to IPerforceSourceControlWorker so that all source control providers can share them.
- ::GetProvider has been removed, if a class inside the module needs access to a provider then a reference to it's provider will be given to it at construction time.
- ::SetLastError, ::AccessSettings and ::SaveSettings are now called directly on FPerforceSourceControlSettings

### FPerforceSourceControlSettings
- This class now takes the name of the system that owns it as a parameter in the constructor
-- The name will be used to identify where to store the settings in the config file, unless the name is blank in which case the pre-existing default location will be used.
- This will allow for multiple source control providers to store their connection info.

### SPerforceSourceControlSettings
- Moved FGetWorkspaces from SPerforceSourceControlSettings.h to its own header file to reduce the number of places that need to include the slate headers.
- Move the fwd declares to the top of the header base don recent coding standards discussions.
- Instead of accessing the source control provider via FPerforceSourceControlModule the dialog is now given a pointer to one when constructed.
-- This pointer is assumed to always be valid and should be a reference if the slate system supported it.

### PerforceSourceControlOperations
- IPerforceSourceControlWorker now takes a reference to the owning FPerforceSourceControlProvider in it's constructor. All derived workers have been updated to this.
- Moved the registering of worker creation delegates from the module cpp to a new static method IPerforceSourceControlWorker::RegisterWorkers
-- This should make it easier to find when adding a new worker and allows source control providers to share the same map of delegates rather than duplicating it.
-- This means we also need a new static method IPerforceSourceControlWorker::CreateWorker for instantiating new workers based on the given name.

### FSourceControlInitSettings
- This is a new class that is used to set up a newly created source control provider.
- The settings are expected to be given in the same way you would override settings via the cmdline. The existing perforce cmdline code has been updated to use this class.
- This still has the downside that the calling code needs to know the type of source control provider, but I don't really see a good way around this unless we refactor the entire API to be feature based.

diff --git a/Engine/Plugins/Developer/PerforceSourceControl/Source/PerforceSourceControl/Private/IPerforceSourceControlWorker.h b/Engine/Plugins/Developer/PerforceSourceControl/Source/PerforceSourceControl/Private/IPerforceSourceControlWorker.h
--- a/Engine/Plugins/Developer/PerforceSourceControl/Source/PerforceSourceControl/Private/IPerforceSourceControlWorker.h
+++ b/Engine/Plugins/Developer/PerforceSourceControl/Source/PerforceSourceControl/Private/IPerforceSourceControlWorker.h
@@ -4,9 +4,23 @@
 
 #include "CoreMinimal.h"
 
+class FPerforceSourceControlProvider;
+
 class IPerforceSourceControlWorker
 {
 public:
+	static void RegisterWorkers();
+	static TSharedPtr<class IPerforceSourceControlWorker, ESPMode::ThreadSafe> CreateWorker(const FName& OperationName, FPerforceSourceControlProvider& SCCProvider);
+public:
+
+	IPerforceSourceControlWorker(FPerforceSourceControlProvider& InSourceControlProvider)
+		: SourceControlProvider(InSourceControlProvider)
+	{
+
+	}
+
+	virtual ~IPerforceSourceControlWorker() = default;
+
 	/**
 	 * Name describing the work that this worker does. Used for factory method hookup.
 	 */
@@ -22,6 +36,15 @@ public:
 	 * @returns true if states were updated
 	 */
 	virtual bool UpdateStates() const = 0;
+
+	FPerforceSourceControlProvider& GetSCCProvider() const
+	{
+		return SourceControlProvider;
+	}
+
+private:
+	 
+	FPerforceSourceControlProvider& SourceControlProvider;
 };
 
 typedef TSharedRef<IPerforceSourceControlWorker, ESPMode::ThreadSafe> FPerforceSourceControlWorkerRef;
diff --git a/Engine/Plugins/Developer/PerforceSourceControl/Source/PerforceSourceControl/Private/PerforceConnection.cpp b/Engine/Plugins/Developer/PerforceSourceControl/Source/PerforceSourceControl/Private/PerforceConnection.cpp
--- a/Engine/Plugins/Developer/PerforceSourceControl/Source/PerforceSourceControl/Private/PerforceConnection.cpp
+++ b/Engine/Plugins/Developer/PerforceSourceControl/Source/PerforceSourceControl/Private/PerforceConnection.cpp
@@ -9,8 +9,8 @@
 #include "Logging/MessageLog.h"
 #include "Memory/SharedBuffer.h"
 #include "Misc/Paths.h"
-#include "Modules/ModuleManager.h"
-#include "PerforceSourceControlModule.h"
+#include "PerforceSourceControlProvider.h"
+#include "PerforceSourceControlSettings.h"
 
 #define LOCTEXT_NAMESPACE "PerforceConnection"
 
@@ -351,11 +351,27 @@ public:
 	FOnIsCancelled IsCancelled;
 };
 
+static bool TestLoginConnection(FPerforceSourceControlProvider& SCCProvider, ClientApi& P4Client, bool bIsUnicodeServer, TArray<FText>& OutErrorMessages)
+{
+	FP4RecordSet Records;
+	EP4ClientUserFlags Flags = bIsUnicodeServer ? EP4ClientUserFlags::UnicodeServer : EP4ClientUserFlags::None;
+	FP4ClientUser User(Records, Flags, OutErrorMessages);
+
+	const char* ArgV[] = { "-s" };
+
+	P4Client.SetArgv(1, const_cast<char* const*>(ArgV));
+	P4Client.Run("login", &User);
+
+	SCCProvider.SetLastErrors(OutErrorMessages);
+
+	return OutErrorMessages.IsEmpty();
+}
+
 /**
  * Runs "client" command to test if the connection is actually OK. ClientApi::Init() only checks
  * if it can connect to server, doesn't verify user name nor workspace name.
  */
-static bool TestConnection(ClientApi& P4Client, const FString& ClientSpecName, bool bIsUnicodeServer, TArray<FText>& OutErrorMessages)
+static bool TestClientConnection(FPerforceSourceControlProvider& SCCProvider, ClientApi& P4Client, const FString& ClientSpecName, bool bIsUnicodeServer, TArray<FText>& OutErrorMessages)
 {
 	FP4RecordSet Records;
 	EP4ClientUserFlags Flags = bIsUnicodeServer ? EP4ClientUserFlags::UnicodeServer : EP4ClientUserFlags::None;
@@ -382,7 +398,7 @@ static bool TestConnection(ClientApi& P4Client, const FString& ClientSpecName, b
 	// clean up args
 	delete [] ClientSpecUTF8Name;
 
-	FPerforceSourceControlModule::SetLastErrors(OutErrorMessages);
+	SCCProvider.SetLastErrors(OutErrorMessages);
 
 	// If there are error messages, user name is most likely invalid. Otherwise, make sure workspace actually
 	// exists on server by checking if we have it's update date.
@@ -414,9 +430,10 @@ static bool CheckUnicodeStatus(ClientApi& P4Client, bool& bIsUnicodeServer, TArr
 	return OutErrorMessages.Num() == 0;
 }
 
-FPerforceConnection::FPerforceConnection(const FPerforceConnectionInfo& InConnectionInfo)
-:	bEstablishedConnection(false)
-,	bIsUnicode(false)
+FPerforceConnection::FPerforceConnection(const FPerforceConnectionInfo& InConnectionInfo, FPerforceSourceControlProvider& InSCCProvider)
+	: bEstablishedConnection(false)
+	, bIsUnicode(false)
+	, SCCProvider(InSCCProvider)
 {
 	EstablishConnection(InConnectionInfo);
 }
@@ -426,14 +443,14 @@ FPerforceConnection::~FPerforceConnection()
 	Disconnect();
 }
 
-bool FPerforceConnection::AutoDetectWorkspace(const FPerforceConnectionInfo& InConnectionInfo, FString& OutWorkspaceName)
+bool FPerforceConnection::AutoDetectWorkspace(const FPerforceConnectionInfo& InConnectionInfo, FPerforceSourceControlProvider& SCCProvider, FString& OutWorkspaceName)
 {
 	bool Result = false;
 	FMessageLog SourceControlLog("SourceControl");
 
 	//before even trying to summon the window, try to "smart" connect with the default server/username
 	TArray<FText> ErrorMessages;
-	FPerforceConnection Connection(InConnectionInfo);
+	FPerforceConnection Connection(InConnectionInfo, SCCProvider);
 	TArray<FString> ClientSpecList;
 	Connection.GetWorkspaceList(InConnectionInfo, FOnIsCancelled(), ClientSpecList, ErrorMessages);
 
@@ -488,7 +505,9 @@ bool FPerforceConnection::Login(const FPerforceConnectionInfo& InConnectionInfo)
 	return ErrorMessages.Num() == 0;
 }
 
-bool FPerforceConnection::EnsureValidConnection(FString& InOutServerName, FString& InOutUserName, FString& InOutWorkspaceName, const FPerforceConnectionInfo& InConnectionInfo)
+bool FPerforceConnection::EnsureValidConnection(FString& InOutServerName, FString& InOutUserName, FString& InOutWorkspaceName,
+												const FPerforceConnectionInfo& InConnectionInfo, FPerforceSourceControlProvider& SCCProvider, 
+												EConnectionOptions Options)
 {
 	TRACE_CPUPROFILER_EVENT_SCOPE(FPerforceConnection::EnsureValidConnection);
 
@@ -529,11 +548,12 @@ bool FPerforceConnection::EnsureValidConnection(FString& InOutServerName, FStrin
 		SourceControlLog.Error(LOCTEXT("P4ErrorConnection_FailedToConnect", "P4ERROR: Failed to connect to source control provider."));
 		SourceControlLog.Error(FText::FromString(ANSI_TO_TCHAR(ErrorMessage.Text())));
 		FFormatNamedArguments Arguments;
+		Arguments.Add(TEXT("OwningSystem"), FText::FromString(SCCProvider.GetOwnerName()));
 		Arguments.Add( TEXT("PortName"), FText::FromString(NewServerName) );
 		Arguments.Add( TEXT("UserName"), FText::FromString(NewUserName) );
 		Arguments.Add( TEXT("ClientSpecName"), FText::FromString(NewClientSpecName) );
 		Arguments.Add( TEXT("Ticket"), FText::FromString(InConnectionInfo.Ticket) );
-		SourceControlLog.Error(FText::Format(LOCTEXT("P4ErrorConnection_Details", "Port={PortName}, User={UserName}, ClientSpec={ClientSpecName}, Ticket={Ticket}"), Arguments));
+		SourceControlLog.Error(FText::Format(LOCTEXT("P4ErrorConnection_Details", "OwningSystem={OwningSystem}, Port={PortName}, User={UserName}, ClientSpec={ClientSpecName}, Ticket={Ticket}"), Arguments));
 	}
 
 	// run an info command to determine unicode status
@@ -547,11 +567,12 @@ bool FPerforceConnection::EnsureValidConnection(FString& InOutServerName, FStrin
 			SourceControlLog.Error(LOCTEXT("P4ErrorConnection_CouldNotDetermineUnicodeStatus", "P4ERROR: Could not determine server unicode status."));
 			SourceControlLog.Error(ErrorMessages.Num() > 0 ? ErrorMessages[0] : LOCTEXT("P4ErrorConnection_UnknownError", "Unknown error"));
 			FFormatNamedArguments Arguments;
+			Arguments.Add(TEXT("OwningSystem"), FText::FromString(SCCProvider.GetOwnerName()));
 			Arguments.Add( TEXT("PortName"), FText::FromString(NewServerName) );
 			Arguments.Add( TEXT("UserName"), FText::FromString(NewUserName) );
 			Arguments.Add( TEXT("ClientSpecName"), FText::FromString(NewClientSpecName) );
 			Arguments.Add( TEXT("Ticket"), FText::FromString(InConnectionInfo.Ticket) );
-			SourceControlLog.Error(FText::Format(LOCTEXT("P4ErrorConnection_Details", "Port={PortName}, User={UserName}, ClientSpec={ClientSpecName}, Ticket={Ticket}"), Arguments));
+			SourceControlLog.Error(FText::Format(LOCTEXT("P4ErrorConnection_Details", "OwningSystem={OwningSystem}, Port={PortName}, User={UserName}, ClientSpec={ClientSpecName}, Ticket={Ticket}"), Arguments));
 		}
 		else
 		{
@@ -570,16 +591,40 @@ bool FPerforceConnection::EnsureValidConnection(FString& InOutServerName, FStrin
 		}
 	}
 
+	// Test that we have a valid p4 ticket
+	if (bConnectionOK)
+	{
+		TArray<FText> ErrorMessages;
+		bConnectionOK = TestLoginConnection(SCCProvider, TestP4, bIsUnicodeServer, ErrorMessages);
+
+		if (!bConnectionOK)
+		{
+			FString ServerName = TO_TCHAR(TestP4.GetPort().Text(), bIsUnicodeServer);
+			FString UserName = TO_TCHAR(TestP4.GetUser().Text(), bIsUnicodeServer);
+
+			SourceControlLog.Error(LOCTEXT("P4ErrorConnection_FailedToConnect", "P4ERROR: Failed to connect to source control provider."));
+			SourceControlLog.Error(ErrorMessages.Num() > 0 ? ErrorMessages[0] : LOCTEXT("P4ErrorConnection_InvalidToken", "Unable to log in"));
+			FFormatNamedArguments Arguments;
+			Arguments.Add(TEXT("OwningSystem"), FText::FromString(SCCProvider.GetOwnerName()));
+			Arguments.Add(TEXT("PortName"), FText::FromString(MoveTemp(ServerName)));
+			Arguments.Add(TEXT("UserName"), FText::FromString(MoveTemp(UserName)));
+
+			SourceControlLog.Error(FText::Format(LOCTEXT("P4ErrorConnection_Details", "OwningSystem={OwningSystem}, Port={PortName}, User={UserName}"), Arguments));
+		}
+	}
+
+	const bool bRequireWorkspace = EnumHasAllFlags(Options, EConnectionOptions::WorkspaceOptional);
+
+	// Try to auto detect the client if none were specified and we require one
 	if (bConnectionOK)
 	{
-		// If a client spec was not specified, attempt to auto-detect it here. If the detection is not successful, neither is this connection since we need a client spec.
-		if ( NewClientSpecName.IsEmpty() )
+		if ( bRequireWorkspace && NewClientSpecName.IsEmpty() )
 		{
 			FPerforceConnectionInfo AutoCredentials = InConnectionInfo;
 			AutoCredentials.Port = TO_TCHAR(TestP4.GetPort().Text(), bIsUnicodeServer);
 			AutoCredentials.UserName = TO_TCHAR(TestP4.GetUser().Text(), bIsUnicodeServer);
 
-		 	bConnectionOK = FPerforceConnection::AutoDetectWorkspace(AutoCredentials, NewClientSpecName);
+		 	bConnectionOK = FPerforceConnection::AutoDetectWorkspace(AutoCredentials, SCCProvider, NewClientSpecName);
 			if ( bConnectionOK )
 			{
 				TestP4.SetClient(FROM_TCHAR(*NewClientSpecName, bIsUnicodeServer));
@@ -587,22 +632,24 @@ bool FPerforceConnection::EnsureValidConnection(FString& InOutServerName, FStrin
 		}
 	}
 
-	if (bConnectionOK)
+	// Test that we found a valid client if we require one
+	if (bConnectionOK && bRequireWorkspace)
 	{
 		TArray<FText> ErrorMessages;
 
-		bConnectionOK = TestConnection(TestP4, NewClientSpecName, bIsUnicodeServer, ErrorMessages);
+		bConnectionOK = TestClientConnection(SCCProvider, TestP4, NewClientSpecName, bIsUnicodeServer, ErrorMessages);
+		
 		if (!bConnectionOK)
 		{
-			//Login FAILED
 			SourceControlLog.Error(LOCTEXT("P4ErrorConnection_FailedToConnect", "P4ERROR: Failed to connect to source control provider."));
 			SourceControlLog.Error(ErrorMessages.Num() > 0 ? ErrorMessages[0] : LOCTEXT("P4ErrorConnection_InvalidWorkspace", "Invalid workspace"));
 			FFormatNamedArguments Arguments;
+			Arguments.Add(TEXT("OwningSystem"), FText::FromString(SCCProvider.GetOwnerName()));
 			Arguments.Add( TEXT("PortName"), FText::FromString(NewServerName) );
 			Arguments.Add( TEXT("UserName"), FText::FromString(NewUserName) );
 			Arguments.Add( TEXT("ClientSpecName"), FText::FromString(NewClientSpecName) );
 			Arguments.Add( TEXT("Ticket"), FText::FromString(InConnectionInfo.Ticket) );
-			SourceControlLog.Error(FText::Format(LOCTEXT("P4ErrorConnection_Details", "Port={PortName}, User={UserName}, ClientSpec={ClientSpecName}, Ticket={Ticket}"), Arguments));
+			SourceControlLog.Error(FText::Format(LOCTEXT("P4ErrorConnection_Details", "OwningSystem={OwningSystem}, Port={PortName}, User={UserName}, ClientSpec={ClientSpecName}, Ticket={Ticket}"), Arguments));
 		}
 	}
 
@@ -829,7 +876,7 @@ bool FPerforceConnection::RunCommand(const FString& InCommand, const TArray<FStr
 	// Only report connection related errors to avoid clearing of connection related error messages
 	if (InCommand != TEXT("info"))
 	{
-		FPerforceSourceControlModule::SetLastErrors(OutErrorMessage);
+		SCCProvider.SetLastErrors(OutErrorMessage);
 	}
 
 	if (bInStandardDebugOutput)
@@ -1040,37 +1087,36 @@ FScopedPerforceConnection::FScopedPerforceConnection( FPerforceSourceControlComm
 	: Connection(nullptr)
 	, Concurrency(InCommand.Concurrency)
 {
-	Initialize(InCommand.ConnectionInfo);
+	Initialize(InCommand.ConnectionInfo, InCommand.Worker->GetSCCProvider());
 	if (IsValid())
 	{
 		InCommand.MarkConnectionAsSuccessful();
 	}
 }
 
-FScopedPerforceConnection::FScopedPerforceConnection( EConcurrency::Type InConcurrency, const FPerforceConnectionInfo& InConnectionInfo )
+FScopedPerforceConnection::FScopedPerforceConnection(EConcurrency::Type InConcurrency, FPerforceSourceControlProvider& SCCProvider)
 	: Connection(nullptr)
 	, Concurrency(InConcurrency)
 {
-	Initialize(InConnectionInfo);
+	Initialize(SCCProvider.AccessSettings().GetConnectionInfo(), SCCProvider);
 }
 
-void FScopedPerforceConnection::Initialize( const FPerforceConnectionInfo& InConnectionInfo )
+void FScopedPerforceConnection::Initialize( const FPerforceConnectionInfo& InConnectionInfo, FPerforceSourceControlProvider& SCCProvider)
 {
 	if(Concurrency == EConcurrency::Synchronous)
 	{
 		// Synchronous commands reuse the same persistent connection to reduce
 		// the number of expensive connection attempts.
-		FPerforceSourceControlModule& PerforceSourceControl = FModuleManager::LoadModuleChecked<FPerforceSourceControlModule>("PerforceSourceControl");
-		if ( PerforceSourceControl.GetProvider().EstablishPersistentConnection() )
+		if (SCCProvider.EstablishPersistentConnection() )
 		{
-			Connection = PerforceSourceControl.GetProvider().GetPersistentConnection();
+			Connection = SCCProvider.GetPersistentConnection();
 		}
 	}
 	else
 	{
 		// Async commands form a new connection for each attempt because
 		// using the persistent connection is not threadsafe
-		FPerforceConnection* NewConnection = new FPerforceConnection(InConnectionInfo);
+		FPerforceConnection* NewConnection = new FPerforceConnection(InConnectionInfo, SCCProvider);
 		if ( NewConnection->IsValidConnection() )
 		{
 			Connection = NewConnection;
diff --git a/Engine/Plugins/Developer/PerforceSourceControl/Source/PerforceSourceControl/Private/PerforceConnection.h b/Engine/Plugins/Developer/PerforceSourceControl/Source/PerforceSourceControl/Private/PerforceConnection.h
--- a/Engine/Plugins/Developer/PerforceSourceControl/Source/PerforceSourceControl/Private/PerforceConnection.h
+++ b/Engine/Plugins/Developer/PerforceSourceControl/Source/PerforceSourceControl/Private/PerforceConnection.h
@@ -7,6 +7,8 @@
 #include "PerforceSourceControlCommand.h"
 #include "Memory/SharedBuffer.h"
 
+class FPerforceSourceControlProvider;
+
 /** A map containing result of running Perforce command */
 class FP4Record : public TMap<FString, FString >
 {
@@ -23,18 +25,28 @@ public:
 
 typedef TArray<FP4Record> FP4RecordSet;
 
+/** Options to be used with FPerforceConnection::EnsureValidConnection */
+enum class EConnectionOptions : uint8
+{
+	/** No options are applied */
+	None				= 0,
+	/** The connection does not require a workspace to be considered valid*/
+	WorkspaceOptional	= 1 << 0,
+};
+ENUM_CLASS_FLAGS(EConnectionOptions);
+
 class FPerforceConnection
 {
 public:
 	//This constructor is strictly for internal questions to perforce (get client spec list, etc)
-	FPerforceConnection(const FPerforceConnectionInfo& InConnectionInfo);
+	FPerforceConnection(const FPerforceConnectionInfo& InConnectionInfo, FPerforceSourceControlProvider& InSCCProvider);
 	/** API Specific close of source control project*/
 	~FPerforceConnection();
 
 	/** 
 	 * Attempts to automatically detect the workspace to use based on the working directory
 	 */
-	static bool AutoDetectWorkspace(const FPerforceConnectionInfo& InConnectionInfo, FString& OutWorkspaceName);
+	static bool AutoDetectWorkspace(const FPerforceConnectionInfo& InConnectionInfo, FPerforceSourceControlProvider& SCCProvider, FString& OutWorkspaceName);
 
 	/**
 	 * Static function in charge of making sure the specified connection is valid or requests that data from the user via dialog
@@ -44,7 +56,9 @@ public:
 	 * @param InConnectionInfo		Connection credentials
 	 * @return - true if the connection, whether via dialog or otherwise, is valid.  False if source control should be disabled
 	 */
-	static bool EnsureValidConnection(FString& InOutServerName, FString& InOutUserName, FString& InOutWorkspaceName, const FPerforceConnectionInfo& InConnectionInfo);
+	static bool EnsureValidConnection(	FString& InOutServerName, FString& InOutUserName, FString& InOutWorkspaceName,
+										const FPerforceConnectionInfo& InConnectionInfo, FPerforceSourceControlProvider& SCCProvider, 
+										EConnectionOptions Options = EConnectionOptions::None);
 
 	/**
 	 * Get List of ClientSpecs
@@ -128,6 +142,7 @@ public:
 	FString GetUser();
 
 public:
+	
 	/** Perforce API client object */
 	ClientApi P4Client;
 
@@ -139,6 +154,10 @@ public:
 
 	/** Is this a connection to a unicode server? */ 
 	bool bIsUnicode;
+
+private:
+	/** The source control provider that the connection is working from */
+	FPerforceSourceControlProvider& SCCProvider;
 };
 
 /**
@@ -153,14 +172,14 @@ public:
 	 * used or another connection is established (connections cannot safely be used across different
 	 * threads).
 	 */
-	FScopedPerforceConnection( class FPerforceSourceControlCommand& InCommand );
+	FScopedPerforceConnection( class FPerforceSourceControlCommand& InCommand);
 
 	/** 
 	 * Constructor - establish a connection.
 	 * The concurrency passed in determines whether the persistent connection is used or another 
 	 * connection is established (connections cannot safely be used across different threads).
 	 */
-	FScopedPerforceConnection( EConcurrency::Type InConcurrency, const FPerforceConnectionInfo& InConnectionInfo );
+	FScopedPerforceConnection( EConcurrency::Type InConcurrency, FPerforceSourceControlProvider& SCCProvider);
 
 	/**
 	 * Destructor - disconnect if this is a temporary connection
@@ -181,7 +200,7 @@ public:
 
 private:
 	/** Set up the connection */
-	void Initialize( const FPerforceConnectionInfo& InConnectionInfo );
+	void Initialize( const FPerforceConnectionInfo& InConnectionInfo, FPerforceSourceControlProvider& SCCProvider);
 
 private:
 	/** The perforce connection we are using */
diff --git a/Engine/Plugins/Developer/PerforceSourceControl/Source/PerforceSourceControl/Private/PerforceSourceControlCommand.cpp b/Engine/Plugins/Developer/PerforceSourceControl/Source/PerforceSourceControl/Private/PerforceSourceControlCommand.cpp
--- a/Engine/Plugins/Developer/PerforceSourceControl/Source/PerforceSourceControl/Private/PerforceSourceControlCommand.cpp
+++ b/Engine/Plugins/Developer/PerforceSourceControl/Source/PerforceSourceControl/Private/PerforceSourceControlCommand.cpp
@@ -1,9 +1,11 @@
 // Copyright Epic Games, Inc. All Rights Reserved.
 
 #include "PerforceSourceControlCommand.h"
-#include "Modules/ModuleManager.h"
-#include "PerforceSourceControlModule.h"
+
+#include "IPerforceSourceControlWorker.h"
 #include "PerforceSourceControlChangelist.h"
+#include "PerforceSourceControlProvider.h"
+#include "PerforceSourceControlSettings.h"
 
 FPerforceSourceControlCommand::FPerforceSourceControlCommand(const TSharedRef<class ISourceControlOperation, ESPMode::ThreadSafe>& InOperation, const TSharedRef<class IPerforceSourceControlWorker, ESPMode::ThreadSafe>& InWorker, const FSourceControlOperationComplete& InOperationCompleteDelegate )
 	: Operation(InOperation)
@@ -20,8 +22,7 @@ FPerforceSourceControlCommand::FPerforceSourceControlCommand(const TSharedRef<cl
 {
 	// Grab the providers settings here, so we don't access them once the worker thread is launched
 	check(IsInGameThread() || InOperation->CanBeCalledFromBackgroundThreads());
-	FPerforceSourceControlModule& PerforceSourceControl = FModuleManager::LoadModuleChecked<FPerforceSourceControlModule>( "PerforceSourceControl" );
-	ConnectionInfo = PerforceSourceControl.AccessSettings().GetConnectionInfo();
+	ConnectionInfo = Worker->GetSCCProvider().AccessSettings().GetConnectionInfo();
 }
 
 bool FPerforceSourceControlCommand::DoWork()
diff --git a/Engine/Plugins/Developer/PerforceSourceControl/Source/PerforceSourceControl/Private/PerforceSourceControlInternalOperations.h b/Engine/Plugins/Developer/PerforceSourceControl/Source/PerforceSourceControl/Private/PerforceSourceControlInternalOperations.h
new file mode 100644
--- /dev/null
+++ b/Engine/Plugins/Developer/PerforceSourceControl/Source/PerforceSourceControl/Private/PerforceSourceControlInternalOperations.h
@@ -0,0 +1,29 @@
+// Copyright Epic Games, Inc. All Rights Reserved.
+
+#pragma once
+
+#include "Containers/Array.h"
+#include "Containers/UnrealString.h"
+#include "SourceControlOperationBase.h"
+#include "UObject/NameTypes.h"
+
+// This head contains operations that are not exposed by the public source control
+// module API and so can only be called from within the perforce source control
+// module itself.
+
+/**
+ * Internal-only source control operation for retrieving available workspaces.
+ */
+class FGetWorkspaces : public FSourceControlOperationBase
+{
+public:
+	// ISourceControlOperation interface
+	virtual FName GetName() const override
+	{
+		return "GetWorkspaces";
+	}
+
+public:
+	/** Results are just an array of workspaces */
+	TArray<FString> Results;
+};
diff --git a/Engine/Plugins/Developer/PerforceSourceControl/Source/PerforceSourceControl/Private/PerforceSourceControlLabel.cpp b/Engine/Plugins/Developer/PerforceSourceControl/Source/PerforceSourceControl/Private/PerforceSourceControlLabel.cpp
--- a/Engine/Plugins/Developer/PerforceSourceControl/Source/PerforceSourceControl/Private/PerforceSourceControlLabel.cpp
+++ b/Engine/Plugins/Developer/PerforceSourceControl/Source/PerforceSourceControl/Private/PerforceSourceControlLabel.cpp
@@ -1,12 +1,13 @@
 // Copyright Epic Games, Inc. All Rights Reserved.
 
 #include "PerforceSourceControlLabel.h"
+
+#include "Logging/MessageLog.h"
+#include "PerforceConnection.h"
 #include "PerforceSourceControlPrivate.h"
-#include "Modules/ModuleManager.h"
+#include "PerforceSourceControlProvider.h"
 #include "PerforceSourceControlRevision.h"
-#include "PerforceSourceControlModule.h"
-#include "PerforceConnection.h"
-#include "Logging/MessageLog.h"
+#include "PerforceSourceControlSettings.h"
 
 #define LOCTEXT_NAMESPACE "PerforceSourceControl"
 
@@ -15,7 +16,7 @@ const FString& FPerforceSourceControlLabel::GetName() const
 	return Name;
 }	
 
-static void ParseFilesResults(const FP4RecordSet& InRecords, TArray< TSharedRef<ISourceControlRevision, ESPMode::ThreadSafe> >& OutRevisions, const FString& InClientRoot)
+static void ParseFilesResults(FPerforceSourceControlProvider& SCCProvider, const FP4RecordSet& InRecords, const FString& InClientRoot, TArray< TSharedRef<ISourceControlRevision, ESPMode::ThreadSafe> >& OutRevisions)
 {
 	// Iterate over each record found as a result of the command, parsing it for relevant information
 	for (int32 Index = 0; Index < InRecords.Num(); ++Index)
@@ -30,7 +31,7 @@ static void ParseFilesResults(const FP4RecordSet& InRecords, TArray< TSharedRef<
 
 		// @todo: this revision is incomplete, but is sufficient for now given the usage of labels to get files, rather
 		// than review revision histories.
-		TSharedRef<FPerforceSourceControlRevision, ESPMode::ThreadSafe> Revision = MakeShareable( new FPerforceSourceControlRevision() );
+		TSharedRef<FPerforceSourceControlRevision, ESPMode::ThreadSafe> Revision = MakeShareable(new FPerforceSourceControlRevision(SCCProvider));
 		Revision->FileName = DepotFile;
 		Revision->RevisionNumber = FCString::Atoi(*RevisionNumber);
 		Revision->ChangelistNumber = FCString::Atoi(*ChangelistNumber);
@@ -41,12 +42,17 @@ static void ParseFilesResults(const FP4RecordSet& InRecords, TArray< TSharedRef<
 	}
 }
 
+FPerforceSourceControlLabel::FPerforceSourceControlLabel(FPerforceSourceControlProvider& InSCCProvider, const FString& InName)
+	: SCCProvider(InSCCProvider)
+	, Name(InName)
+{
+}
+
 bool FPerforceSourceControlLabel::GetFileRevisions( const TArray<FString>& InFiles, TArray< TSharedRef<ISourceControlRevision, ESPMode::ThreadSafe> >& OutRevisions ) const
 {
 	bool bCommandOK = false;
 
-	FPerforceSourceControlModule& PerforceSourceControl = FModuleManager::LoadModuleChecked<FPerforceSourceControlModule>("PerforceSourceControl");
-	FScopedPerforceConnection ScopedConnection(EConcurrency::Synchronous, PerforceSourceControl.AccessSettings().GetConnectionInfo());
+	FScopedPerforceConnection ScopedConnection(EConcurrency::Synchronous, GetSCCProvider());
 	if(ScopedConnection.IsValid())
 	{
 		FPerforceConnection& Connection = ScopedConnection.GetConnection();
@@ -61,7 +67,7 @@ bool FPerforceSourceControlLabel::GetFileRevisions( const TArray<FString>& InFil
 		bCommandOK = Connection.RunCommand(TEXT("files"), Parameters, Records, ErrorMessages, FOnIsCancelled(), bConnectionDropped);
 		if(bCommandOK)
 		{
-			ParseFilesResults(Records, OutRevisions, Connection.ClientRoot);
+			ParseFilesResults(GetSCCProvider(), Records, Connection.ClientRoot, OutRevisions);
 		}
 		else
 		{
@@ -80,8 +86,7 @@ bool FPerforceSourceControlLabel::Sync( const TArray<FString>& InFilenames ) con
 {
 	bool bCommandOK = false;
 
-	FPerforceSourceControlModule& PerforceSourceControl = FModuleManager::LoadModuleChecked<FPerforceSourceControlModule>("PerforceSourceControl");
-	FScopedPerforceConnection ScopedConnection(EConcurrency::Synchronous, PerforceSourceControl.AccessSettings().GetConnectionInfo());
+	FScopedPerforceConnection ScopedConnection(EConcurrency::Synchronous, GetSCCProvider());
 	if(ScopedConnection.IsValid())
 	{
 		FPerforceConnection& Connection = ScopedConnection.GetConnection();
diff --git a/Engine/Plugins/Developer/PerforceSourceControl/Source/PerforceSourceControl/Private/PerforceSourceControlLabel.h b/Engine/Plugins/Developer/PerforceSourceControl/Source/PerforceSourceControl/Private/PerforceSourceControlLabel.h
--- a/Engine/Plugins/Developer/PerforceSourceControl/Source/PerforceSourceControl/Private/PerforceSourceControlLabel.h
+++ b/Engine/Plugins/Developer/PerforceSourceControl/Source/PerforceSourceControl/Private/PerforceSourceControlLabel.h
@@ -6,6 +6,8 @@
 #include "ISourceControlRevision.h"
 #include "ISourceControlLabel.h"
 
+class FPerforceSourceControlProvider;
+
 /** 
  * Abstraction of a Perforce label.
  */
@@ -13,10 +15,8 @@ class FPerforceSourceControlLabel : public ISourceControlLabel, public TSharedFr
 {
 public:
 
-	FPerforceSourceControlLabel( const FString& InName )
-		: Name(InName)
-	{
-	}
+	FPerforceSourceControlLabel(FPerforceSourceControlProvider& InSCCProvider, const FString& InName);
+	virtual ~FPerforceSourceControlLabel() = default;
 
 	/** ISourceControlLabel implementation */
 	virtual const FString& GetName() const override;
@@ -25,6 +25,15 @@ public:
 
 private:
 
+	/** Internal accessor to the source control provider associated with the object */
+	FPerforceSourceControlProvider& GetSCCProvider() const
+	{
+		return SCCProvider;
+	}
+
+	/** The source control provider that this object is associated with */
+	FPerforceSourceControlProvider& SCCProvider;
+
 	/** Label name */
 	FString Name;
 };
diff --git a/Engine/Plugins/Developer/PerforceSourceControl/Source/PerforceSourceControl/Private/PerforceSourceControlModule.cpp b/Engine/Plugins/Developer/PerforceSourceControl/Source/PerforceSourceControl/Private/PerforceSourceControlModule.cpp
--- a/Engine/Plugins/Developer/PerforceSourceControl/Source/PerforceSourceControl/Private/PerforceSourceControlModule.cpp
+++ b/Engine/Plugins/Developer/PerforceSourceControl/Source/PerforceSourceControl/Private/PerforceSourceControlModule.cpp
@@ -1,49 +1,18 @@
 // Copyright Epic Games, Inc. All Rights Reserved.
 
 #include "PerforceSourceControlModule.h"
+
 #include "Misc/App.h"
 #include "PerforceSourceControlOperations.h"
 #include "Features/IModularFeatures.h"
+#include "PerforceSourceControlSettings.h"
 
 #define LOCTEXT_NAMESPACE "PerforceSourceControl"
 
-template<typename Type>
-static TSharedRef<IPerforceSourceControlWorker, ESPMode::ThreadSafe> CreateWorker()
-{
-	return MakeShareable( new Type() );
-}
-
 void FPerforceSourceControlModule::StartupModule()
 {
 	// Register our operations
-	PerforceSourceControlProvider.RegisterWorker( "Connect", FGetPerforceSourceControlWorker::CreateStatic( &CreateWorker<FPerforceConnectWorker> ) );
-	PerforceSourceControlProvider.RegisterWorker( "CheckOut", FGetPerforceSourceControlWorker::CreateStatic( &CreateWorker<FPerforceCheckOutWorker> ) );
-	PerforceSourceControlProvider.RegisterWorker( "UpdateStatus", FGetPerforceSourceControlWorker::CreateStatic( &CreateWorker<FPerforceUpdateStatusWorker> ) );
-	PerforceSourceControlProvider.RegisterWorker( "MarkForAdd", FGetPerforceSourceControlWorker::CreateStatic( &CreateWorker<FPerforceMarkForAddWorker> ) );
-	PerforceSourceControlProvider.RegisterWorker( "Delete", FGetPerforceSourceControlWorker::CreateStatic( &CreateWorker<FPerforceDeleteWorker> ) );
-	PerforceSourceControlProvider.RegisterWorker( "Revert", FGetPerforceSourceControlWorker::CreateStatic( &CreateWorker<FPerforceRevertWorker> ) );
-	PerforceSourceControlProvider.RegisterWorker( "Sync", FGetPerforceSourceControlWorker::CreateStatic( &CreateWorker<FPerforceSyncWorker> ) );
-	PerforceSourceControlProvider.RegisterWorker( "CheckIn", FGetPerforceSourceControlWorker::CreateStatic( &CreateWorker<FPerforceCheckInWorker> ) );
-	PerforceSourceControlProvider.RegisterWorker( "GetWorkspaces", FGetPerforceSourceControlWorker::CreateStatic( &CreateWorker<FPerforceGetWorkspacesWorker> ) );
-	PerforceSourceControlProvider.RegisterWorker( "Copy", FGetPerforceSourceControlWorker::CreateStatic( &CreateWorker<FPerforceCopyWorker> ) );
-	PerforceSourceControlProvider.RegisterWorker( "Resolve", FGetPerforceSourceControlWorker::CreateStatic( &CreateWorker<FPerforceResolveWorker> ) );
-	PerforceSourceControlProvider.RegisterWorker( "ChangeStatus", FGetPerforceSourceControlWorker::CreateStatic(&CreateWorker<FPerforceChangeStatusWorker>));
-	PerforceSourceControlProvider.RegisterWorker( "UpdateChangelistsStatus", FGetPerforceSourceControlWorker::CreateStatic( &CreateWorker<FPerforceGetPendingChangelistsWorker> ) );
-	PerforceSourceControlProvider.RegisterWorker( "NewChangelist", FGetPerforceSourceControlWorker::CreateStatic( &CreateWorker<FPerforceNewChangelistWorker> ) );
-	PerforceSourceControlProvider.RegisterWorker( "DeleteChangelist", FGetPerforceSourceControlWorker::CreateStatic( &CreateWorker<FPerforceDeleteChangelistWorker> ) );
-	PerforceSourceControlProvider.RegisterWorker( "EditChangelist", FGetPerforceSourceControlWorker::CreateStatic( &CreateWorker<FPerforceEditChangelistWorker> ) );
-	PerforceSourceControlProvider.RegisterWorker( "RevertUnchanged", FGetPerforceSourceControlWorker::CreateStatic( &CreateWorker<FPerforceRevertUnchangedWorker> ) );
-	PerforceSourceControlProvider.RegisterWorker( "MoveToChangelist", FGetPerforceSourceControlWorker::CreateStatic( &CreateWorker<FPerforceReopenWorker>) );
-	PerforceSourceControlProvider.RegisterWorker( "Shelve", FGetPerforceSourceControlWorker::CreateStatic( &CreateWorker<FPerforceShelveWorker>) );
-	PerforceSourceControlProvider.RegisterWorker( "Unshelve", FGetPerforceSourceControlWorker::CreateStatic( &CreateWorker<FPerforceUnshelveWorker>) );
-	PerforceSourceControlProvider.RegisterWorker( "DeleteShelved", FGetPerforceSourceControlWorker::CreateStatic( &CreateWorker<FPerforceDeleteShelveWorker>) );
-	PerforceSourceControlProvider.RegisterWorker( "DownloadFile", FGetPerforceSourceControlWorker::CreateStatic( &CreateWorker<FPerforceDownloadFileWorker>) );
-	PerforceSourceControlProvider.RegisterWorker( "CreateWorkspace", FGetPerforceSourceControlWorker::CreateStatic( &CreateWorker<FPerforceCreateWorkspaceWorker>) );
-	PerforceSourceControlProvider.RegisterWorker( "DeleteWorkspace", FGetPerforceSourceControlWorker::CreateStatic( &CreateWorker<FPerforceDeleteWorkspaceWorker>) );
-	PerforceSourceControlProvider.RegisterWorker( "GetFileList", FGetPerforceSourceControlWorker::CreateStatic(&CreateWorker<FPerforceGetFileListWorker>));
-
-	// load our settings
-	PerforceSourceControlSettings.LoadSettings();
+	IPerforceSourceControlWorker::RegisterWorkers();
 
 	// Bind our source control provider to the editor
 	IModularFeatures::Get().RegisterModularFeature( "SourceControl", &PerforceSourceControlProvider );
@@ -58,30 +27,6 @@ void FPerforceSourceControlModule::ShutdownModule()
 	IModularFeatures::Get().UnregisterModularFeature( "SourceControl", &PerforceSourceControlProvider );
 }
 
-FPerforceSourceControlSettings& FPerforceSourceControlModule::AccessSettings()
-{
-	return PerforceSourceControlSettings;
-}
-
-void FPerforceSourceControlModule::SaveSettings()
-{
-	if (FApp::IsUnattended() || IsRunningCommandlet())
-	{
-		return;
-	}
-
-	PerforceSourceControlSettings.SaveSettings();
-}
-
-void FPerforceSourceControlModule::SetLastErrors(const TArray<FText>& InErrors)
-{
-	FPerforceSourceControlModule* Module = FModuleManager::GetModulePtr<FPerforceSourceControlModule>("PerforceSourceControl");
-	if (Module)
-	{
-		Module->GetProvider().SetLastErrors(InErrors);
-	}
-}
-
 IMPLEMENT_MODULE(FPerforceSourceControlModule, PerforceSourceControl);
 
 #undef LOCTEXT_NAMESPACE
diff --git a/Engine/Plugins/Developer/PerforceSourceControl/Source/PerforceSourceControl/Private/PerforceSourceControlModule.h b/Engine/Plugins/Developer/PerforceSourceControl/Source/PerforceSourceControl/Private/PerforceSourceControlModule.h
--- a/Engine/Plugins/Developer/PerforceSourceControl/Source/PerforceSourceControl/Private/PerforceSourceControlModule.h
+++ b/Engine/Plugins/Developer/PerforceSourceControl/Source/PerforceSourceControl/Private/PerforceSourceControlModule.h
@@ -6,7 +6,6 @@
 #include "Modules/ModuleInterface.h"
 #include "Modules/ModuleManager.h"
 #include "PerforceSourceControlProvider.h"
-#include "PerforceSourceControlSettings.h"
 
 class FPerforceSourceControlModule : public IModuleInterface
 {
@@ -14,37 +13,8 @@ public:
 	/** IModuleInterface implementation */
 	virtual void StartupModule() override;
 	virtual void ShutdownModule() override;
-
-	/** Access the Perforce source control settings */
-	FPerforceSourceControlSettings& AccessSettings();
-
-	/** Save the Perforce source control settings */
-	void SaveSettings();
-
-	/** Access the one and only Perforce provider */
-	FPerforceSourceControlProvider& GetProvider()
-	{
-		return PerforceSourceControlProvider;
-	}
 	
-	/**
-	 * Singleton-like access to this module's interface.  This is just for convenience!
-	 * Beware of calling this during the shutdown phase, though.  Your module might have been unloaded already.
-	 *
-	 * @return Returns singleton instance, loading the module on demand if needed
-	 */
-	static inline FPerforceSourceControlModule& Get()
-	{
-		return FModuleManager::LoadModuleChecked< FPerforceSourceControlModule >("PerforceSourceControl");
-	}
-
-	/** Set list of error messages that occurred after last perforce command */
-	static void SetLastErrors(const TArray<FText>& InErrors);
-
 private:
-	/** The one and only Perforce source control provider */
+	/** The perforce source control provider that will be used by default via ISourceControlModule::Get().GetProvider() etc */
 	FPerforceSourceControlProvider PerforceSourceControlProvider;
-
-	/** The settings for Perforce source control */
-	FPerforceSourceControlSettings PerforceSourceControlSettings;
 };
diff --git a/Engine/Plugins/Developer/PerforceSourceControl/Source/PerforceSourceControl/Private/PerforceSourceControlOperations.cpp b/Engine/Plugins/Developer/PerforceSourceControl/Source/PerforceSourceControl/Private/PerforceSourceControlOperations.cpp
--- a/Engine/Plugins/Developer/PerforceSourceControl/Source/PerforceSourceControl/Private/PerforceSourceControlOperations.cpp
+++ b/Engine/Plugins/Developer/PerforceSourceControl/Source/PerforceSourceControl/Private/PerforceSourceControlOperations.cpp
@@ -1,29 +1,81 @@
 // Copyright Epic Games, Inc. All Rights Reserved.
 
 #include "PerforceSourceControlOperations.h"
-#include "PerforceSourceControlPrivate.h"
+
+#include "Algo/AnyOf.h"
+#include "Algo/Find.h"
+#include "Algo/IndexOf.h"
+#include "Algo/Transform.h"
 #include "HAL/FileManager.h"
 #include "ISourceControlModule.h"
-#include "Misc/Paths.h"
 #include "Misc/EngineVersion.h"
-#include "Modules/ModuleManager.h"
-#include "ISourceControlModule.h"
-#include "SourceControlHelpers.h"
-#include "SourceControlOperations.h"
-#include "PerforceSourceControlRevision.h"
-#include "PerforceSourceControlCommand.h"
-#include "PerforceSourceControlChangelistState.h"
+#include "Misc/Paths.h"
 #include "PerforceConnection.h"
-#include "PerforceSourceControlModule.h"
 #include "PerforceSourceControlChangeStatusOperation.h"
-#include "SPerforceSourceControlSettings.h"
-#include "Algo/AnyOf.h"
-#include "Algo/IndexOf.h"
-#include "Algo/Find.h"
-#include "Algo/Transform.h"
+#include "PerforceSourceControlChangelistState.h"
+#include "PerforceSourceControlCommand.h"
+#include "PerforceSourceControlInternalOperations.h"
+#include "PerforceSourceControlPrivate.h"
+#include "PerforceSourceControlProvider.h"
+#include "PerforceSourceControlRevision.h"
+#include "SourceControlHelpers.h"
+#include "SourceControlOperations.h"
 
 #define LOCTEXT_NAMESPACE "PerforceSourceControl"
 
+namespace
+{
+DECLARE_DELEGATE_RetVal_OneParam(FPerforceSourceControlWorkerRef, FGetPerforceSourceControlWorker, FPerforceSourceControlProvider&)
+
+static TMap<FName, FGetPerforceSourceControlWorker> WorkersMap;
+}
+
+template<typename Type>
+FPerforceSourceControlWorkerRef InstantiateWorker(FPerforceSourceControlProvider& InSourceControlProvider)
+{
+	return MakeShareable(new Type(InSourceControlProvider));
+}
+
+void IPerforceSourceControlWorker::RegisterWorkers()
+{
+	WorkersMap.Add("Connect", FGetPerforceSourceControlWorker::CreateStatic(&InstantiateWorker<FPerforceConnectWorker>));
+	WorkersMap.Add("CheckOut", FGetPerforceSourceControlWorker::CreateStatic(&InstantiateWorker<FPerforceCheckOutWorker>));
+	WorkersMap.Add("UpdateStatus", FGetPerforceSourceControlWorker::CreateStatic(&InstantiateWorker<FPerforceUpdateStatusWorker>));
+	WorkersMap.Add("MarkForAdd", FGetPerforceSourceControlWorker::CreateStatic(&InstantiateWorker<FPerforceMarkForAddWorker>));
+	WorkersMap.Add("Delete", FGetPerforceSourceControlWorker::CreateStatic(&InstantiateWorker<FPerforceDeleteWorker>));
+	WorkersMap.Add("Revert", FGetPerforceSourceControlWorker::CreateStatic(&InstantiateWorker<FPerforceRevertWorker>));
+	WorkersMap.Add("Sync", FGetPerforceSourceControlWorker::CreateStatic(&InstantiateWorker<FPerforceSyncWorker>));
+	WorkersMap.Add("CheckIn", FGetPerforceSourceControlWorker::CreateStatic(&InstantiateWorker<FPerforceCheckInWorker>));
+	WorkersMap.Add("GetWorkspaces", FGetPerforceSourceControlWorker::CreateStatic(&InstantiateWorker<FPerforceGetWorkspacesWorker>));
+	WorkersMap.Add("Copy", FGetPerforceSourceControlWorker::CreateStatic(&InstantiateWorker<FPerforceCopyWorker>));
+	WorkersMap.Add("Resolve", FGetPerforceSourceControlWorker::CreateStatic(&InstantiateWorker<FPerforceResolveWorker>));
+	WorkersMap.Add("ChangeStatus", FGetPerforceSourceControlWorker::CreateStatic(&InstantiateWorker<FPerforceChangeStatusWorker>));
+	WorkersMap.Add("UpdateChangelistsStatus", FGetPerforceSourceControlWorker::CreateStatic(&InstantiateWorker<FPerforceGetPendingChangelistsWorker>));
+	WorkersMap.Add("NewChangelist", FGetPerforceSourceControlWorker::CreateStatic(&InstantiateWorker<FPerforceNewChangelistWorker>));
+	WorkersMap.Add("DeleteChangelist", FGetPerforceSourceControlWorker::CreateStatic(&InstantiateWorker<FPerforceDeleteChangelistWorker>));
+	WorkersMap.Add("EditChangelist", FGetPerforceSourceControlWorker::CreateStatic(&InstantiateWorker<FPerforceEditChangelistWorker>));
+	WorkersMap.Add("RevertUnchanged", FGetPerforceSourceControlWorker::CreateStatic(&InstantiateWorker<FPerforceRevertUnchangedWorker>));
+	WorkersMap.Add("MoveToChangelist", FGetPerforceSourceControlWorker::CreateStatic(&InstantiateWorker<FPerforceReopenWorker>));
+	WorkersMap.Add("Shelve", FGetPerforceSourceControlWorker::CreateStatic(&InstantiateWorker<FPerforceShelveWorker>));
+	WorkersMap.Add("Unshelve", FGetPerforceSourceControlWorker::CreateStatic(&InstantiateWorker<FPerforceUnshelveWorker>));
+	WorkersMap.Add("DeleteShelved", FGetPerforceSourceControlWorker::CreateStatic(&InstantiateWorker<FPerforceDeleteShelveWorker>));
+	WorkersMap.Add("DownloadFile", FGetPerforceSourceControlWorker::CreateStatic(&InstantiateWorker<FPerforceDownloadFileWorker>));
+	WorkersMap.Add("CreateWorkspace", FGetPerforceSourceControlWorker::CreateStatic(&InstantiateWorker<FPerforceCreateWorkspaceWorker>));
+	WorkersMap.Add("DeleteWorkspace", FGetPerforceSourceControlWorker::CreateStatic(&InstantiateWorker<FPerforceDeleteWorkspaceWorker>));
+	WorkersMap.Add("GetFileList", FGetPerforceSourceControlWorker::CreateStatic(&InstantiateWorker<FPerforceGetFileListWorker>));
+}
+
+TSharedPtr<class IPerforceSourceControlWorker, ESPMode::ThreadSafe> IPerforceSourceControlWorker::CreateWorker(const FName& OperationName, FPerforceSourceControlProvider& SCCProvider)
+{
+	const FGetPerforceSourceControlWorker* Operation = WorkersMap.Find(OperationName);
+	if (Operation != nullptr)
+	{
+		return Operation->Execute(SCCProvider);
+	}
+
+	return nullptr;
+}
+
 /**
  * Helper struct for RemoveRedundantErrors()
  */
@@ -186,12 +238,11 @@ static void ParseRecordSetForState(const FP4RecordSet& InRecords, TMap<FString,
 	}
 }
 
-static void UpdateChangelistState(const FPerforceSourceControlChangelist& InChangelist, const TMap<FString, EPerforceState::Type>& InOperationResults)
+static void UpdateChangelistState(FPerforceSourceControlProvider& SCCProvider, const FPerforceSourceControlChangelist& InChangelist, const TMap<FString, EPerforceState::Type>& InOperationResults)
 {
 	if (InChangelist.IsInitialized())
 	{
-		FPerforceSourceControlModule& PerforceSourceControl = FPerforceSourceControlModule::Get();
-		TSharedRef<FPerforceSourceControlChangelistState, ESPMode::ThreadSafe> ChangelistState = PerforceSourceControl.GetProvider().GetStateInternal(InChangelist);
+		TSharedRef<FPerforceSourceControlChangelistState, ESPMode::ThreadSafe> ChangelistState = SCCProvider.GetStateInternal(InChangelist);
 
 		for (TMap<FString, EPerforceState::Type>::TConstIterator It(InOperationResults); It; ++It)
 		{
@@ -200,18 +251,17 @@ static void UpdateChangelistState(const FPerforceSourceControlChangelist& InChan
 				continue;
 			}
 
-			TSharedRef<FPerforceSourceControlState, ESPMode::ThreadSafe> State = PerforceSourceControl.GetProvider().GetStateInternal(It.Key());
+			TSharedRef<FPerforceSourceControlState, ESPMode::ThreadSafe> State = SCCProvider.GetStateInternal(It.Key());
 			ChangelistState->Files.Add(State);
 		}
 	}
 }
 
-static bool UpdateCachedStates(const TMap<FString, EPerforceState::Type>& InResults)
+static bool UpdateCachedStates(FPerforceSourceControlProvider& SCCProvider, const TMap<FString, EPerforceState::Type>& InResults)
 {
-	FPerforceSourceControlModule& PerforceSourceControl = FPerforceSourceControlModule::Get();
 	for(TMap<FString, EPerforceState::Type>::TConstIterator It(InResults); It; ++It)
 	{
-		TSharedRef<FPerforceSourceControlState, ESPMode::ThreadSafe> State = PerforceSourceControl.GetProvider().GetStateInternal(It.Key());
+		TSharedRef<FPerforceSourceControlState, ESPMode::ThreadSafe> State = SCCProvider.GetStateInternal(It.Key());
 		State->SetState(It.Value());
 		State->TimeStamp = FDateTime::Now();
 	}
@@ -255,10 +305,9 @@ static bool CheckWorkspaceRecordSet(const FP4RecordSet& InRecords, TArray<FText>
 	return false;
 }
 
-static bool AppendChangelistParameter(TArray<FString>& InOutParams)
+static bool AppendChangelistParameter(FPerforceSourceControlProvider& SCCProvider, TArray<FString>& InOutParams)
 {
-	FPerforceSourceControlModule& PerforceSourceControl = FModuleManager::GetModuleChecked<FPerforceSourceControlModule>("PerforceSourceControl");
-	FPerforceSourceControlSettings& Settings = PerforceSourceControl.AccessSettings();
+	FPerforceSourceControlSettings& Settings = SCCProvider.AccessSettings();
 
 	const FString& ChangelistNumber = Settings.GetChangelistNumber();
 	if ( !ChangelistNumber.IsEmpty() )
@@ -336,7 +385,7 @@ bool FPerforceCheckOutWorker::Execute(FPerforceSourceControlCommand& InCommand)
 		FPerforceConnection& Connection = ScopedConnection.GetConnection();
 		TArray<FString> Parameters;
 
-		if ((!AppendChangelistParameter(Parameters)) && InCommand.Changelist.IsInitialized())
+		if ((!AppendChangelistParameter(GetSCCProvider(), Parameters)) && InCommand.Changelist.IsInitialized())
 		{
 			FString ChangelistNumber = InCommand.Changelist.ToString();
 
@@ -360,9 +409,9 @@ bool FPerforceCheckOutWorker::Execute(FPerforceSourceControlCommand& InCommand)
 bool FPerforceCheckOutWorker::UpdateStates() const
 {
 	// If files have been checkedout directly to a CL, modify the cached state to reflect it.
-	UpdateChangelistState(InChangelist, OutResults);
+	UpdateChangelistState(GetSCCProvider(), InChangelist, OutResults);
 
-	return UpdateCachedStates(OutResults);
+	return UpdateCachedStates(GetSCCProvider(), OutResults);
 }
 
 FName FPerforceCheckInWorker::GetName() const
@@ -443,10 +492,9 @@ static bool RemoveFilesFromChangelist(const TMap<FString, EPerforceState::Type>&
 		}) > 0;
 }
 
-static bool RemoveFilesFromChangelist(const TMap<FString, EPerforceState::Type>& Results, const FPerforceSourceControlChangelist& Changelist)
+static bool RemoveFilesFromChangelist(const TMap<FString, EPerforceState::Type>& Results, FPerforceSourceControlProvider& SCCProvider, const FPerforceSourceControlChangelist& Changelist)
 {
-	FPerforceSourceControlModule& PerforceSourceControl = FPerforceSourceControlModule::Get();
-	TSharedRef<FPerforceSourceControlChangelistState, ESPMode::ThreadSafe> ChangelistState = PerforceSourceControl.GetProvider().GetStateInternal(Changelist);
+	TSharedRef<FPerforceSourceControlChangelistState, ESPMode::ThreadSafe> ChangelistState = SCCProvider.GetStateInternal(Changelist);
 	return RemoveFilesFromChangelist(Results, ChangelistState);
 }
 
@@ -457,16 +505,14 @@ bool FPerforceCheckInWorker::Execute(FPerforceSourceControlCommand& InCommand)
 	{
 		FPerforceConnection& Connection = ScopedConnection.GetConnection();
 		FText CachedDescription;
-		FPerforceSourceControlModule& PerforceSourceControl = FModuleManager::GetModuleChecked<FPerforceSourceControlModule>("PerforceSourceControl");
-		FPerforceSourceControlProvider& Provider = PerforceSourceControl.GetProvider();
-
+		
 		check(InCommand.Operation->GetName() == GetName());
 		TSharedRef<FCheckIn, ESPMode::ThreadSafe> Operation = StaticCastSharedRef<FCheckIn>(InCommand.Operation);
 
 		TArray<FString> FilesToSubmit = InCommand.Files;
 
 		FPerforceSourceControlChangelist ChangeList(InCommand.Changelist);
-		FSourceControlChangelistStateRef ChangelistState = PerforceSourceControl.GetProvider().GetStateInternal(InCommand.Changelist);
+		FSourceControlChangelistStateRef ChangelistState = GetSCCProvider().GetStateInternal(InCommand.Changelist);
 		TArray<FString> ReopenedFiles;
 
 		InCommand.bCommandSuccessful = true;
@@ -527,12 +573,12 @@ bool FPerforceCheckInWorker::Execute(FPerforceSourceControlCommand& InCommand)
 			{
 				// Remove any deleted files from status cache
 				TArray<TSharedRef<ISourceControlState, ESPMode::ThreadSafe>> States;
-				Provider.GetState(FilesToSubmit, States, EStateCacheUsage::Use);
+				GetSCCProvider().GetState(FilesToSubmit, States, EStateCacheUsage::Use);
 				for (const auto& State : States)
 				{
 					if (State->IsDeleted())
 					{
-						Provider.RemoveFileFromCache(State->GetFilename());
+						GetSCCProvider().RemoveFileFromCache(State->GetFilename());
 					}
 				}
 
@@ -581,19 +627,18 @@ bool FPerforceCheckInWorker::Execute(FPerforceSourceControlCommand& InCommand)
 
 bool FPerforceCheckInWorker::UpdateStates() const
 {
-	bool bUpdatedStates = UpdateCachedStates(OutResults);
+	bool bUpdatedStates = UpdateCachedStates(GetSCCProvider(), OutResults);
 	bool bUpdatedChangelistStates = false;
 
 	if(!OutChangelist.IsDefault()) // e.g. operation succeeded
 	{
 		// Delete changelist, whether its a temporary one or not
-		FPerforceSourceControlModule& PerforceSourceControl = FPerforceSourceControlModule::Get();
-		bUpdatedChangelistStates = PerforceSourceControl.GetProvider().RemoveChangelistFromCache(OutChangelist);
+		bUpdatedChangelistStates = GetSCCProvider().RemoveChangelistFromCache(OutChangelist);
 
 		// If it's a temporary one, then remove the submitted files from the default changelist
 		if (InChangelist.IsDefault())
 		{
-			bUpdatedChangelistStates = RemoveFilesFromChangelist(OutResults, InChangelist);
+			bUpdatedChangelistStates = RemoveFilesFromChangelist(OutResults, GetSCCProvider(), InChangelist);
 		}
 	}
 
@@ -700,7 +745,7 @@ bool FPerforceMarkForAddWorker::Execute(FPerforceSourceControlCommand& InCommand
 			}
 		}
 
-		AppendChangelistParameter(Parameters);
+		AppendChangelistParameter(GetSCCProvider(), Parameters);
 		Parameters.Append(InCommand.Files);
 
 		InCommand.bCommandSuccessful = Connection.RunCommand(TEXT("add"), Parameters, Records, InCommand.ResultInfo.ErrorMessages, FOnIsCancelled::CreateRaw(&InCommand, &FPerforceSourceControlCommand::IsCanceled), InCommand.bConnectionDropped);
@@ -711,9 +756,9 @@ bool FPerforceMarkForAddWorker::Execute(FPerforceSourceControlCommand& InCommand
 
 bool FPerforceMarkForAddWorker::UpdateStates() const
 {
-	UpdateChangelistState(InChangelist, OutResults);
+	UpdateChangelistState(GetSCCProvider(), InChangelist, OutResults);
 
-	return UpdateCachedStates(OutResults);
+	return UpdateCachedStates(GetSCCProvider(), OutResults);
 }
 
 FName FPerforceDeleteWorker::GetName() const
@@ -729,7 +774,7 @@ bool FPerforceDeleteWorker::Execute(FPerforceSourceControlCommand& InCommand)
 		FPerforceConnection& Connection = ScopedConnection.GetConnection();
 		TArray<FString> Parameters;
 
-		AppendChangelistParameter(Parameters);
+		AppendChangelistParameter(GetSCCProvider(), Parameters);
 		Parameters.Append(InCommand.Files);
 
 		FP4RecordSet Records;
@@ -741,7 +786,7 @@ bool FPerforceDeleteWorker::Execute(FPerforceSourceControlCommand& InCommand)
 
 bool FPerforceDeleteWorker::UpdateStates() const
 {
-	return UpdateCachedStates(OutResults);
+	return UpdateCachedStates(GetSCCProvider(), OutResults);
 }
 
 FName FPerforceRevertWorker::GetName() const
@@ -764,7 +809,7 @@ bool FPerforceRevertWorker::Execute(FPerforceSourceControlCommand& InCommand)
 		}
 		else
 		{
-			AppendChangelistParameter(Parameters);
+			AppendChangelistParameter(GetSCCProvider(), Parameters);
 		}
 
 		TSharedRef<FRevert> RevertOperation = StaticCastSharedRef<FRevert>(InCommand.Operation);
@@ -806,18 +851,17 @@ bool FPerforceRevertWorker::Execute(FPerforceSourceControlCommand& InCommand)
 
 bool FPerforceRevertWorker::UpdateStates() const
 {
-	FPerforceSourceControlModule& PerforceSourceControl = FPerforceSourceControlModule::Get();
-	bool bUpdatedCachedStates = UpdateCachedStates(OutResults);
+	bool bUpdatedCachedStates = UpdateCachedStates(GetSCCProvider(), OutResults);
 	bool bUpdatedChangelists = false;
 
-	bUpdatedChangelists = ChangelistToUpdate.IsInitialized() && RemoveFilesFromChangelist(OutResults, ChangelistToUpdate);
+	bUpdatedChangelists = ChangelistToUpdate.IsInitialized() && RemoveFilesFromChangelist(OutResults, GetSCCProvider(), ChangelistToUpdate);
 
 	// Use reverted files to update changelist state.
 	for (TMap<FString, EPerforceState::Type>::TConstIterator It(OutResults); It; ++It)
 	{
-		TSharedRef<FPerforceSourceControlState, ESPMode::ThreadSafe> State = PerforceSourceControl.GetProvider().GetStateInternal(It.Key());
+		TSharedRef<FPerforceSourceControlState, ESPMode::ThreadSafe> State = GetSCCProvider().GetStateInternal(It.Key());
 
-		bUpdatedChangelists |= State->Changelist.IsInitialized() && RemoveFilesFromChangelist(OutResults, State->Changelist);
+		bUpdatedChangelists |= State->Changelist.IsInitialized() && RemoveFilesFromChangelist(OutResults, GetSCCProvider(), State->Changelist);
 	}
 
 	return bUpdatedCachedStates || bUpdatedChangelists;
@@ -857,8 +901,7 @@ bool FPerforceSyncWorker::Execute(FPerforceSourceControlCommand& InCommand)
 	if (!InCommand.IsCanceled() && ScopedConnection.IsValid())
 	{
 		FPerforceConnection& Connection = ScopedConnection.GetConnection();
-		FPerforceSourceControlModule& PerforceSourceControl = FPerforceSourceControlModule::Get();
-
+		
 		TSharedRef<FSync, ESPMode::ThreadSafe> Operation = StaticCastSharedRef<FSync>(InCommand.Operation);
 		TArray<FString> Parameters;
 
@@ -883,7 +926,7 @@ bool FPerforceSyncWorker::Execute(FPerforceSourceControlCommand& InCommand)
 			}
 			else if (Operation->IsLastSyncedFlagSet())
 			{
-				TSharedRef<FPerforceSourceControlState, ESPMode::ThreadSafe> State = PerforceSourceControl.GetProvider().GetStateInternal(FileName);
+				TSharedRef<FPerforceSourceControlState, ESPMode::ThreadSafe> State = GetSCCProvider().GetStateInternal(FileName);
 
 				if (State->LocalRevNumber == ISourceControlState::INVALID_REVISION)
 				{
@@ -917,7 +960,7 @@ bool FPerforceSyncWorker::Execute(FPerforceSourceControlCommand& InCommand)
 
 bool FPerforceSyncWorker::UpdateStates() const
 {
-	return UpdateCachedStates(OutResults);
+	return UpdateCachedStates(GetSCCProvider(), OutResults);
 }
 
 static void ParseBranchModificationResults(const FP4RecordSet& InRecords, const TArray<FText>& ErrorMessages, const FString& ContentRoot, TMap<FString, FBranchModification>& BranchModifications)
@@ -1463,7 +1506,7 @@ static const FString& FindWorkspaceFile(const TArray<FPerforceSourceControlState
 	return InDepotFile;
 }
 
-static void ParseHistoryResults(const FP4RecordSet& InRecords, const TArray<FPerforceSourceControlState>& InStates, FPerforceFileHistoryMap& OutHistory)
+static void ParseHistoryResults(FPerforceSourceControlProvider& SCCProvider, const FP4RecordSet& InRecords, const TArray<FPerforceSourceControlState>& InStates, FPerforceFileHistoryMap& OutHistory)
 {
 	if (InRecords.Num() > 0)
 	{
@@ -1541,7 +1584,7 @@ static void ParseHistoryResults(const FP4RecordSet& InRecords, const TArray<FPer
 				VarName = FString::Printf(TEXT("how%d,0"), RevisionNumbers);
 				if(ClientRecord.Contains(*VarName))
 				{
-					BranchSource = MakeShareable( new FPerforceSourceControlRevision() );
+					BranchSource = MakeShareable(new FPerforceSourceControlRevision(SCCProvider));
 
 					VarName = FString::Printf(TEXT("file%d,0"), RevisionNumbers);
 					FString BranchSourceFileName = ClientRecord(*VarName);
@@ -1552,7 +1595,7 @@ static void ParseHistoryResults(const FP4RecordSet& InRecords, const TArray<FPer
 					BranchSource->RevisionNumber = FCString::Atoi(*BranchSourceRevision);
 				}
 
-				TSharedRef<FPerforceSourceControlRevision, ESPMode::ThreadSafe> Revision = MakeShareable( new FPerforceSourceControlRevision() );
+				TSharedRef<FPerforceSourceControlRevision, ESPMode::ThreadSafe> Revision = MakeShareable(new FPerforceSourceControlRevision(SCCProvider));
 				Revision->FileName = LocalFileName;
 				Revision->RevisionNumber = FCString::Atoi(*RevisionNumber);
 				Revision->Revision = RevisionNumber;
@@ -1578,7 +1621,8 @@ static void ParseHistoryResults(const FP4RecordSet& InRecords, const TArray<FPer
 	}
 }
 
-static bool GetFileHistory(FPerforceConnection& Connection, FPerforceSourceControlCommand& InCommand, const TArray<FString>& InFiles, TArray<FPerforceSourceControlState>& OutStates, FPerforceFileHistoryMap& OutHistory)
+static bool GetFileHistory(	FPerforceSourceControlProvider& SCCProvider, FPerforceConnection& Connection, FPerforceSourceControlCommand& InCommand, 
+							const TArray<FString>& InFiles, TArray<FPerforceSourceControlState>& OutStates, FPerforceFileHistoryMap& OutHistory)
 {
 	TArray<FString> Parameters;
 	FP4RecordSet Records;
@@ -1594,7 +1638,7 @@ static bool GetFileHistory(FPerforceConnection& Connection, FPerforceSourceContr
 	Parameters.Add(TEXT("-m 100"));
 	Parameters.Append(InFiles);
 	InCommand.bCommandSuccessful &= Connection.RunCommand(TEXT("filelog"), Parameters, Records, InCommand.ResultInfo.ErrorMessages, FOnIsCancelled::CreateRaw(&InCommand, &FPerforceSourceControlCommand::IsCanceled), InCommand.bConnectionDropped);
-	ParseHistoryResults(Records, OutStates, OutHistory);
+	ParseHistoryResults(SCCProvider, Records, OutStates, OutHistory);
 	RemoveRedundantErrors(InCommand, TEXT(" - no such file(s)."));
 	RemoveRedundantErrors(InCommand, TEXT(" - file(s) not on client"));
 	RemoveRedundantErrors(InCommand, TEXT("' is not under client's root '"));
@@ -1730,7 +1774,7 @@ bool FPerforceUpdateStatusWorker::Execute(FPerforceSourceControlCommand& InComma
 
 		if(Operation->ShouldUpdateHistory())
 		{
-			GetFileHistory(Connection, InCommand, InCommand.Files, OutStates, OutHistory);
+			GetFileHistory(GetSCCProvider(), Connection, InCommand, InCommand.Files, OutStates, OutHistory);
 		}
 
 		if(Operation->ShouldGetOpenedOnly())
@@ -1790,14 +1834,13 @@ bool FPerforceUpdateStatusWorker::UpdateStates() const
 {
 	bool bUpdated = false;
 
-	FPerforceSourceControlModule& PerforceSourceControl = FPerforceSourceControlModule::Get();
 	const FDateTime Now = FDateTime::Now();
 
 	// first update cached state from 'fstat' call
 	for(int StatusIndex = 0; StatusIndex < OutStates.Num(); StatusIndex++)
 	{
 		const FPerforceSourceControlState& Status = OutStates[StatusIndex];
-		TSharedRef<FPerforceSourceControlState, ESPMode::ThreadSafe> State = PerforceSourceControl.GetProvider().GetStateInternal(Status.LocalFilename);
+		TSharedRef<FPerforceSourceControlState, ESPMode::ThreadSafe> State = GetSCCProvider().GetStateInternal(Status.LocalFilename);
 		// Update every member except History and Timestamp. History will be updated below from the OutHistory map.
 		// Timestamp is used to throttle status requests, so update it to current time:
 		auto History = MoveTemp(State->History);
@@ -1808,12 +1851,12 @@ bool FPerforceUpdateStatusWorker::UpdateStates() const
 	}
 
 	// next update state from 'opened' call
-	bUpdated |= UpdateCachedStates(OutStateMap);
+	bUpdated |= UpdateCachedStates(GetSCCProvider(), OutStateMap);
 
 	// add history, if any
 	for(FPerforceFileHistoryMap::TConstIterator It(OutHistory); It; ++It)
 	{
-		TSharedRef<FPerforceSourceControlState, ESPMode::ThreadSafe> State = PerforceSourceControl.GetProvider().GetStateInternal(It.Key());
+		TSharedRef<FPerforceSourceControlState, ESPMode::ThreadSafe> State = GetSCCProvider().GetStateInternal(It.Key());
 		const TArray< TSharedRef<FPerforceSourceControlRevision, ESPMode::ThreadSafe> >& History = It.Value();
 		State->History = History;
 		State->TimeStamp = Now;
@@ -1824,7 +1867,7 @@ bool FPerforceUpdateStatusWorker::UpdateStates() const
 	for(int ModifiedIndex = 0; ModifiedIndex < OutModifiedFiles.Num(); ModifiedIndex++)
 	{
 		const FString& FileName = OutModifiedFiles[ModifiedIndex];
-		TSharedRef<FPerforceSourceControlState, ESPMode::ThreadSafe> State = PerforceSourceControl.GetProvider().GetStateInternal(FileName);
+		TSharedRef<FPerforceSourceControlState, ESPMode::ThreadSafe> State = GetSCCProvider().GetStateInternal(FileName);
 		State->bModifed = true;
 		State->TimeStamp = Now;
 		bUpdated = true;
@@ -2026,9 +2069,10 @@ bool FPerforceGetPendingChangelistsWorker::Execute(FPerforceSourceControlCommand
 	return InCommand.bCommandSuccessful;
 }
 
-static bool AddShelvedFilesToChangelist(const TMap<FString, EPerforceState::Type>& FilesToAdd, const TMap<FString, FString>& DepotToFileMap, TSharedRef<FPerforceSourceControlChangelistState, ESPMode::ThreadSafe>& ChangelistState, const FDateTime* TimeStamp = nullptr)
+static bool AddShelvedFilesToChangelist(FPerforceSourceControlProvider& SCCProvider, const TMap<FString, EPerforceState::Type>& FilesToAdd, 
+										const TMap<FString, FString>& DepotToFileMap, TSharedRef<FPerforceSourceControlChangelistState, ESPMode::ThreadSafe>& ChangelistState, 
+										const FDateTime* TimeStamp = nullptr)
 {
-	FPerforceSourceControlModule& PerforceSourceControl = FPerforceSourceControlModule::Get();
 	const FDateTime Now = (TimeStamp ? *TimeStamp : FDateTime::Now());
 
 	for (TMap<FString, EPerforceState::Type>::TConstIterator It(FilesToAdd); It; ++It)
@@ -2054,7 +2098,7 @@ static bool AddShelvedFilesToChangelist(const TMap<FString, EPerforceState::Type
 			// Add revision to be able to fetch the shelved file, if it's not marked for deletion.
 			if (It.Value() != EPerforceState::MarkedForDelete)
 			{
-				TSharedRef<FPerforceSourceControlRevision, ESPMode::ThreadSafe> ShelvedRevision = MakeShareable(new FPerforceSourceControlRevision());
+				TSharedRef<FPerforceSourceControlRevision, ESPMode::ThreadSafe> ShelvedRevision = MakeShareable(new FPerforceSourceControlRevision(SCCProvider));
 				ShelvedRevision->FileName = ShelvedFileState->DepotFilename;
 				ShelvedRevision->ChangelistNumber = StaticCastSharedRef<FPerforceSourceControlChangelist>(ChangelistState->GetChangelist())->ToInt();
 				ShelvedRevision->bIsShelve = true;
@@ -2076,25 +2120,24 @@ static bool AddShelvedFilesToChangelist(const TMap<FString, EPerforceState::Type
 	return FilesToAdd.Num() > 0;
 }
 
-static bool AddShelvedFilesToChangelist(const TMap<FString, EPerforceState::Type>& FilesToAdd, const TMap<FString, FString>& DepotToFileMap, const FPerforceSourceControlChangelist& Changelist)
+static bool AddShelvedFilesToChangelist(FPerforceSourceControlProvider& SCCProvider, const TMap<FString, EPerforceState::Type>& FilesToAdd, 
+										const TMap<FString, FString>& DepotToFileMap, const FPerforceSourceControlChangelist& Changelist)
 {
-	FPerforceSourceControlModule& PerforceSourceControl = FPerforceSourceControlModule::Get();
-	TSharedRef<FPerforceSourceControlChangelistState, ESPMode::ThreadSafe> ChangelistState = PerforceSourceControl.GetProvider().GetStateInternal(Changelist);
-	return AddShelvedFilesToChangelist(FilesToAdd, DepotToFileMap, ChangelistState);
+	TSharedRef<FPerforceSourceControlChangelistState, ESPMode::ThreadSafe> ChangelistState = SCCProvider.GetStateInternal(Changelist);
+	return AddShelvedFilesToChangelist(SCCProvider, FilesToAdd, DepotToFileMap, ChangelistState);
 }
 
 bool FPerforceGetPendingChangelistsWorker::UpdateStates() const
 {
 	bool bUpdated = false;
 
-	FPerforceSourceControlModule& PerforceSourceControl = FPerforceSourceControlModule::Get();
 	const FDateTime Now = FDateTime::Now();
 
 	// first update cached state from 'changes' call
 	for (int StatusIndex = 0; StatusIndex < OutChangelistsStates.Num(); StatusIndex++)
 	{
 		const FPerforceSourceControlChangelistState& CLStatus = OutChangelistsStates[StatusIndex];
-		TSharedRef<FPerforceSourceControlChangelistState, ESPMode::ThreadSafe> ChangelistState = PerforceSourceControl.GetProvider().GetStateInternal(CLStatus.Changelist);
+		TSharedRef<FPerforceSourceControlChangelistState, ESPMode::ThreadSafe> ChangelistState = GetSCCProvider().GetStateInternal(CLStatus.Changelist);
 		// Timestamp is used to throttle status requests, so update it to current time:
 		*ChangelistState = CLStatus;
 		ChangelistState->TimeStamp = Now;
@@ -2107,7 +2150,7 @@ bool FPerforceGetPendingChangelistsWorker::UpdateStates() const
 			ChangelistState->Files.Reset(OutCLFilesStates[StatusIndex].Num());
 			for (const auto& FileState : OutCLFilesStates[StatusIndex])
 			{
-				TSharedRef<FPerforceSourceControlState, ESPMode::ThreadSafe> CachedFileState = PerforceSourceControl.GetProvider().GetStateInternal(FileState.LocalFilename);
+				TSharedRef<FPerforceSourceControlState, ESPMode::ThreadSafe> CachedFileState = GetSCCProvider().GetStateInternal(FileState.LocalFilename);
 				CachedFileState->Update(FileState, &Now);
 				ChangelistState->Files.AddUnique(CachedFileState);
 			}
@@ -2118,14 +2161,14 @@ bool FPerforceGetPendingChangelistsWorker::UpdateStates() const
 		if(bUpdateShelvedFiles)
 		{
 			ChangelistState->ShelvedFiles.Reset(OutCLShelvedFilesStates[StatusIndex].Num());
-			AddShelvedFilesToChangelist(OutCLShelvedFilesStates[StatusIndex], OutCLShelvedFilesMap[StatusIndex], ChangelistState, &Now);
+			AddShelvedFilesToChangelist(GetSCCProvider(), OutCLShelvedFilesStates[StatusIndex], OutCLShelvedFilesMap[StatusIndex], ChangelistState, &Now);
 		}
 	}
 
 	if (bCleanupCache)
 	{
 		TArray<FPerforceSourceControlChangelist> ChangelistsToRemove;
-		PerforceSourceControl.GetProvider().GetCachedStateByPredicate([this, &ChangelistsToRemove](const FSourceControlChangelistStateRef& InCLState) {
+		GetSCCProvider().GetCachedStateByPredicate([this, &ChangelistsToRemove](const FSourceControlChangelistStateRef& InCLState) {
 			TSharedRef<FPerforceSourceControlChangelistState, ESPMode::ThreadSafe> CLState = StaticCastSharedRef<FPerforceSourceControlChangelistState>(InCLState);
 
 			if (Algo::NoneOf(OutChangelistsStates, [&CLState](const FPerforceSourceControlChangelistState& UpdatedCLState) {
@@ -2140,7 +2183,7 @@ bool FPerforceGetPendingChangelistsWorker::UpdateStates() const
 
 		for (const FPerforceSourceControlChangelist& ChangelistToRemove : ChangelistsToRemove)
 		{
-			PerforceSourceControl.GetProvider().RemoveChangelistFromCache(ChangelistToRemove);
+			GetSCCProvider().RemoveChangelistFromCache(ChangelistToRemove);
 		}
 	}
 
@@ -2167,7 +2210,7 @@ bool FPerforceCopyWorker::Execute(class FPerforceSourceControlCommand& InCommand
 
 		TArray<FString> Parameters;
 
-		AppendChangelistParameter(Parameters);
+		AppendChangelistParameter(GetSCCProvider(), Parameters);
 
 		Parameters.Append(InCommand.Files);
 		Parameters.Add(DestinationPath);
@@ -2192,7 +2235,7 @@ bool FPerforceCopyWorker::Execute(class FPerforceSourceControlCommand& InCommand
 
 bool FPerforceCopyWorker::UpdateStates() const
 {
-	return UpdateCachedStates(OutResults);
+	return UpdateCachedStates(GetSCCProvider(), OutResults);
 }
 
 // IPerforceSourceControlWorker interface
@@ -2212,7 +2255,7 @@ bool FPerforceResolveWorker::Execute(class FPerforceSourceControlCommand& InComm
 
 		Parameters.Add("-ay");
 		Parameters.Append(InCommand.Files);
-		AppendChangelistParameter(Parameters);
+		AppendChangelistParameter(GetSCCProvider(), Parameters);
 
 		FP4RecordSet Records;
 		InCommand.bCommandSuccessful = Connection.RunCommand(TEXT("resolve"), Parameters, Records, InCommand.ResultInfo.ErrorMessages, FOnIsCancelled::CreateRaw(&InCommand, &FPerforceSourceControlCommand::IsCanceled), InCommand.bConnectionDropped);
@@ -2227,11 +2270,9 @@ bool FPerforceResolveWorker::Execute(class FPerforceSourceControlCommand& InComm
 
 bool FPerforceResolveWorker::UpdateStates() const
 {
-	FPerforceSourceControlModule& PerforceSourceControl = FPerforceSourceControlModule::Get();
-
 	for( const auto& Filename : UpdatedFiles )
 	{
-		auto State = PerforceSourceControl.GetProvider().GetStateInternal( Filename );
+		auto State = GetSCCProvider().GetStateInternal( Filename );
 		State->LocalRevNumber = State->DepotRevNumber;
 		State->PendingResolveRevNumber = FPerforceSourceControlState::INVALID_REVISION;
 	}
@@ -2287,8 +2328,9 @@ bool FPerforceChangeStatusWorker::UpdateStates() const
 	return true;
 }
 
-FPerforceNewChangelistWorker::FPerforceNewChangelistWorker()
-	: NewChangelistState(NewChangelist)
+FPerforceNewChangelistWorker::FPerforceNewChangelistWorker(FPerforceSourceControlProvider& InSourceControlProvider)
+	: IPerforceSourceControlWorker(InSourceControlProvider)
+	, NewChangelistState(NewChangelist)
 {
 
 }
@@ -2348,19 +2390,18 @@ bool FPerforceNewChangelistWorker::UpdateStates() const
 {
 	if (NewChangelist.IsInitialized())
 	{
-		FPerforceSourceControlModule& PerforceSourceControl = FPerforceSourceControlModule::Get();
 		const FDateTime Now = FDateTime::Now();
 
-		TSharedRef<FPerforceSourceControlChangelistState, ESPMode::ThreadSafe> ChangelistState = PerforceSourceControl.GetProvider().GetStateInternal(NewChangelist);
+		TSharedRef<FPerforceSourceControlChangelistState, ESPMode::ThreadSafe> ChangelistState = GetSCCProvider().GetStateInternal(NewChangelist);
 		*ChangelistState = NewChangelistState;
 		ChangelistState->TimeStamp = Now;
 
 		for (const FString& MovedFile : MovedFiles)
 		{
-			TSharedRef<FPerforceSourceControlState, ESPMode::ThreadSafe> FileState = PerforceSourceControl.GetProvider().GetStateInternal(MovedFile);
+			TSharedRef<FPerforceSourceControlState, ESPMode::ThreadSafe> FileState = GetSCCProvider().GetStateInternal(MovedFile);
 
 			// 1- Remove these files from their previous changelist
-			TSharedRef<FPerforceSourceControlChangelistState, ESPMode::ThreadSafe> PreviousChangelist = PerforceSourceControl.GetProvider().GetStateInternal(FileState->Changelist);
+			TSharedRef<FPerforceSourceControlChangelistState, ESPMode::ThreadSafe> PreviousChangelist = GetSCCProvider().GetStateInternal(FileState->Changelist);
 			PreviousChangelist->Files.Remove(FileState);
 
 			// 2- Add to the new changelist
@@ -2421,10 +2462,9 @@ bool FPerforceDeleteChangelistWorker::Execute(class FPerforceSourceControlComman
 
 bool FPerforceDeleteChangelistWorker::UpdateStates() const
 {
-	FPerforceSourceControlModule& PerforceSourceControl = FPerforceSourceControlModule::Get();
 	if (!DeletedChangelist.IsDefault())
 	{
-		return PerforceSourceControl.GetProvider().RemoveChangelistFromCache(DeletedChangelist);
+		return GetSCCProvider().RemoveChangelistFromCache(DeletedChangelist);
 	}
 	else
 	{
@@ -2472,8 +2512,7 @@ bool FPerforceEditChangelistWorker::Execute(class FPerforceSourceControlCommand&
 
 bool FPerforceEditChangelistWorker::UpdateStates() const
 {
-	FPerforceSourceControlModule& PerforceSourceControl = FPerforceSourceControlModule::Get();
-	TSharedRef<FPerforceSourceControlChangelistState, ESPMode::ThreadSafe> EditedChangelistState = PerforceSourceControl.GetProvider().GetStateInternal(EditedChangelist);
+	TSharedRef<FPerforceSourceControlChangelistState, ESPMode::ThreadSafe> EditedChangelistState = GetSCCProvider().GetStateInternal(EditedChangelist);
 	// TODO: update similar to NewChangelist when/if we support files in edit/new changelists.
 	EditedChangelistState->Description = EditedDescription.ToString();
 	EditedChangelistState->Changelist = EditedChangelist;
@@ -2514,8 +2553,8 @@ bool FPerforceRevertUnchangedWorker::Execute(class FPerforceSourceControlCommand
 
 bool FPerforceRevertUnchangedWorker::UpdateStates() const
 {
-	bool bUpdatedStates = UpdateCachedStates(OutResults);
-	bool bUpdatedChangelistState = ChangelistToUpdate.IsInitialized() && RemoveFilesFromChangelist(OutResults, ChangelistToUpdate);
+	bool bUpdatedStates = UpdateCachedStates(GetSCCProvider(), OutResults);
+	bool bUpdatedChangelistState = ChangelistToUpdate.IsInitialized() && RemoveFilesFromChangelist(OutResults, GetSCCProvider(), ChangelistToUpdate);
 	return bUpdatedStates || bUpdatedChangelistState;
 }
 
@@ -2540,16 +2579,16 @@ bool FPerforceReopenWorker::Execute(FPerforceSourceControlCommand& InCommand)
 bool FPerforceReopenWorker::UpdateStates() const
 {
 	const FDateTime Now = FDateTime::Now();
-	FPerforceSourceControlModule& PerforceSourceControl = FPerforceSourceControlModule::Get();
-	TSharedRef<FPerforceSourceControlChangelistState, ESPMode::ThreadSafe> DestinationChangelistState = PerforceSourceControl.GetProvider().GetStateInternal(DestinationChangelist);
+	
+	TSharedRef<FPerforceSourceControlChangelistState, ESPMode::ThreadSafe> DestinationChangelistState = GetSCCProvider().GetStateInternal(DestinationChangelist);
 
 	// 3 things to do here:
 	for (const FString& ReopenedFile : ReopenedFiles)
 	{
-		TSharedRef<FPerforceSourceControlState, ESPMode::ThreadSafe> FileState = PerforceSourceControl.GetProvider().GetStateInternal(ReopenedFile);
+		TSharedRef<FPerforceSourceControlState, ESPMode::ThreadSafe> FileState = GetSCCProvider().GetStateInternal(ReopenedFile);
 
 		// 1- Remove these files from their previous changelist
-		TSharedRef<FPerforceSourceControlChangelistState, ESPMode::ThreadSafe> PreviousChangelist = PerforceSourceControl.GetProvider().GetStateInternal(FileState->Changelist);
+		TSharedRef<FPerforceSourceControlChangelistState, ESPMode::ThreadSafe> PreviousChangelist = GetSCCProvider().GetStateInternal(FileState->Changelist);
 		PreviousChangelist->Files.Remove(FileState);
 
 		// 2- Add to the new changelist
@@ -2592,8 +2631,7 @@ bool FPerforceShelveWorker::Execute(class FPerforceSourceControlCommand& InComma
 			// If the command has specified the default changelist but no files, then get all files from the default changelist
 			if (FilesToShelve.Num() == 0 && InCommand.Changelist.IsInitialized())
 			{
-				FPerforceSourceControlModule& PerforceSourceControl = FPerforceSourceControlModule::Get();
-				TSharedRef<FPerforceSourceControlChangelistState, ESPMode::ThreadSafe> DefaultChangelistState = PerforceSourceControl.GetProvider().GetStateInternal(InCommand.Changelist);
+				TSharedRef<FPerforceSourceControlChangelistState, ESPMode::ThreadSafe> DefaultChangelistState = GetSCCProvider().GetStateInternal(InCommand.Changelist);
 				Algo::Transform(DefaultChangelistState->Files, FilesToShelve, [](const auto& FileState) {
 					return FileState->GetFilename();
 					});
@@ -2663,16 +2701,14 @@ bool FPerforceShelveWorker::Execute(class FPerforceSourceControlCommand& InComma
 
 bool FPerforceShelveWorker::UpdateStates() const
 {
-	FPerforceSourceControlModule& PerforceSourceControl = FPerforceSourceControlModule::Get();
-
 	bool bMovedFiles = false;
 
 	// If we moved files to a new changelist, then we must make sure that the files are properly moved
 	if (InChangelistToUpdate != OutChangelistToUpdate && MovedFiles.Num() > 0)
 	{
 		const FDateTime Now = FDateTime::Now();
-		TSharedRef<FPerforceSourceControlChangelistState, ESPMode::ThreadSafe> SourceChangelistState = PerforceSourceControl.GetProvider().GetStateInternal(InChangelistToUpdate);
-		TSharedRef<FPerforceSourceControlChangelistState, ESPMode::ThreadSafe> DestinationChangelistState = PerforceSourceControl.GetProvider().GetStateInternal(OutChangelistToUpdate);
+		TSharedRef<FPerforceSourceControlChangelistState, ESPMode::ThreadSafe> SourceChangelistState = GetSCCProvider().GetStateInternal(InChangelistToUpdate);
+		TSharedRef<FPerforceSourceControlChangelistState, ESPMode::ThreadSafe> DestinationChangelistState = GetSCCProvider().GetStateInternal(OutChangelistToUpdate);
 
 		DestinationChangelistState->Changelist = OutChangelistToUpdate;
 		DestinationChangelistState->Description = ChangelistDescription;
@@ -2680,7 +2716,7 @@ bool FPerforceShelveWorker::UpdateStates() const
 
 		for (const FString& MovedFile : MovedFiles)
 		{
-			TSharedRef<FPerforceSourceControlState, ESPMode::ThreadSafe> FileState = PerforceSourceControl.GetProvider().GetStateInternal(MovedFile);
+			TSharedRef<FPerforceSourceControlState, ESPMode::ThreadSafe> FileState = GetSCCProvider().GetStateInternal(MovedFile);
 
 			SourceChangelistState->Files.Remove(FileState);
 			DestinationChangelistState->Files.Add(FileState);
@@ -2691,7 +2727,7 @@ bool FPerforceShelveWorker::UpdateStates() const
 		bMovedFiles = true;
 	}
 
-	const bool bAddedShelvedFilesToChangelist = (OutResults.Num() > 0 && AddShelvedFilesToChangelist(OutResults, OutFileMap, OutChangelistToUpdate));
+	const bool bAddedShelvedFilesToChangelist = (OutResults.Num() > 0 && AddShelvedFilesToChangelist(GetSCCProvider(), OutResults, OutFileMap, OutChangelistToUpdate));
 
 	return bMovedFiles || bAddedShelvedFilesToChangelist;
 }
@@ -2736,8 +2772,7 @@ bool FPerforceDeleteShelveWorker::UpdateStates() const
 {
 	if (ChangelistToUpdate.IsInitialized())
 	{
-		FPerforceSourceControlModule& PerforceSourceControl = FPerforceSourceControlModule::Get();
-		TSharedRef<FPerforceSourceControlChangelistState, ESPMode::ThreadSafe> ChangelistState = PerforceSourceControl.GetProvider().GetStateInternal(ChangelistToUpdate);
+		TSharedRef<FPerforceSourceControlChangelistState, ESPMode::ThreadSafe> ChangelistState = GetSCCProvider().GetStateInternal(ChangelistToUpdate);
 
 		if (FilesToRemove.Num() > 0)
 		{
@@ -2809,13 +2844,12 @@ bool FPerforceUnshelveWorker::UpdateStates() const
 	{
 		const FDateTime Now = FDateTime::Now();
 
-		FPerforceSourceControlModule& PerforceSourceControl = FPerforceSourceControlModule::Get();
-		TSharedRef<FPerforceSourceControlChangelistState, ESPMode::ThreadSafe> ChangelistState = PerforceSourceControl.GetProvider().GetStateInternal(ChangelistToUpdate);
+		TSharedRef<FPerforceSourceControlChangelistState, ESPMode::ThreadSafe> ChangelistState = GetSCCProvider().GetStateInternal(ChangelistToUpdate);
 
 		ChangelistState->Files.Reset(ChangelistFilesStates.Num());
 		for (const auto& FileState : ChangelistFilesStates)
 		{
-			TSharedRef<FPerforceSourceControlState, ESPMode::ThreadSafe> CachedFileState = PerforceSourceControl.GetProvider().GetStateInternal(FileState.LocalFilename);
+			TSharedRef<FPerforceSourceControlState, ESPMode::ThreadSafe> CachedFileState = GetSCCProvider().GetStateInternal(FileState.LocalFilename);
 			CachedFileState->Update(FileState, &Now);
 
 			ChangelistState->Files.AddUnique(CachedFileState);
diff --git a/Engine/Plugins/Developer/PerforceSourceControl/Source/PerforceSourceControl/Private/PerforceSourceControlOperations.h b/Engine/Plugins/Developer/PerforceSourceControl/Source/PerforceSourceControl/Private/PerforceSourceControlOperations.h
--- a/Engine/Plugins/Developer/PerforceSourceControl/Source/PerforceSourceControl/Private/PerforceSourceControlOperations.h
+++ b/Engine/Plugins/Developer/PerforceSourceControl/Source/PerforceSourceControl/Private/PerforceSourceControlOperations.h
@@ -10,20 +10,28 @@
 class FPerforceSourceControlRevision;
 typedef TMap<FString, TArray< TSharedRef<FPerforceSourceControlRevision, ESPMode::ThreadSafe> > > FPerforceFileHistoryMap;
 
-class FPerforceConnectWorker : public IPerforceSourceControlWorker
+class FPerforceConnectWorker final : public IPerforceSourceControlWorker
 {
 public:
-	virtual ~FPerforceConnectWorker() {}
+	FPerforceConnectWorker(FPerforceSourceControlProvider& InSourceControlProvider) 
+		: IPerforceSourceControlWorker(InSourceControlProvider) 
+	{}
+	virtual ~FPerforceConnectWorker() = default;
+
 	// IPerforceSourceControlWorker interface
 	virtual FName GetName() const override;
 	virtual bool Execute(class FPerforceSourceControlCommand& InCommand) override;
 	virtual bool UpdateStates() const override;
 };
 
-class FPerforceCheckOutWorker : public IPerforceSourceControlWorker
+class FPerforceCheckOutWorker final : public IPerforceSourceControlWorker
 {
 public:
-	virtual ~FPerforceCheckOutWorker() {}
+	FPerforceCheckOutWorker(FPerforceSourceControlProvider& InSourceControlProvider)
+		: IPerforceSourceControlWorker(InSourceControlProvider)
+	{}
+	virtual ~FPerforceCheckOutWorker() = default;
+
 	// IPerforceSourceControlWorker interface
 	virtual FName GetName() const override;
 	virtual bool Execute(class FPerforceSourceControlCommand& InCommand) override;
@@ -35,10 +43,14 @@ public:
 	FPerforceSourceControlChangelist InChangelist;
 };
 
-class FPerforceCheckInWorker : public IPerforceSourceControlWorker
+class FPerforceCheckInWorker final : public IPerforceSourceControlWorker
 {
 public:
-	virtual ~FPerforceCheckInWorker() {}
+	FPerforceCheckInWorker(FPerforceSourceControlProvider& InSourceControlProvider)
+		: IPerforceSourceControlWorker(InSourceControlProvider)
+	{}
+	virtual ~FPerforceCheckInWorker() = default;
+
 	// IPerforceSourceControlWorker interface
 	virtual FName GetName() const override;
 	virtual bool Execute(class FPerforceSourceControlCommand& InCommand) override;
@@ -55,10 +67,14 @@ public:
 	FPerforceSourceControlChangelist OutChangelist;
 };
 
-class FPerforceGetFileListWorker : public IPerforceSourceControlWorker
+class FPerforceGetFileListWorker final : public IPerforceSourceControlWorker
 {
 public:
-	virtual ~FPerforceGetFileListWorker() {}
+	FPerforceGetFileListWorker(FPerforceSourceControlProvider& InSourceControlProvider)
+		: IPerforceSourceControlWorker(InSourceControlProvider)
+	{}
+	virtual ~FPerforceGetFileListWorker() = default;
+
 	// IPerforceSourceControlWorker interface
 	virtual FName GetName() const override;
 	virtual bool Execute(class FPerforceSourceControlCommand& InCommand) override;
@@ -69,10 +85,14 @@ public:
 	TMap<FString, EPerforceState::Type> OutResults;
 };
 
-class FPerforceMarkForAddWorker : public IPerforceSourceControlWorker
+class FPerforceMarkForAddWorker final : public IPerforceSourceControlWorker
 {
 public:
-	virtual ~FPerforceMarkForAddWorker() {}
+	FPerforceMarkForAddWorker(FPerforceSourceControlProvider& InSourceControlProvider)
+		: IPerforceSourceControlWorker(InSourceControlProvider)
+	{}
+	virtual ~FPerforceMarkForAddWorker() = default;
+
 	// IPerforceSourceControlWorker interface
 	virtual FName GetName() const override;
 	virtual bool Execute(class FPerforceSourceControlCommand& InCommand) override;
@@ -84,10 +104,14 @@ public:
 	FPerforceSourceControlChangelist InChangelist;
 };
 
-class FPerforceDeleteWorker : public IPerforceSourceControlWorker
+class FPerforceDeleteWorker final : public IPerforceSourceControlWorker
 {
 public:
-	virtual ~FPerforceDeleteWorker() {}
+	FPerforceDeleteWorker(FPerforceSourceControlProvider& InSourceControlProvider)
+		: IPerforceSourceControlWorker(InSourceControlProvider)
+	{}
+	virtual ~FPerforceDeleteWorker() = default;
+
 	// IPerforceSourceControlWorker interface
 	virtual FName GetName() const override;
 	virtual bool Execute(class FPerforceSourceControlCommand& InCommand) override;
@@ -98,10 +122,14 @@ public:
 	TMap<FString, EPerforceState::Type> OutResults;
 };
 
-class FPerforceRevertWorker : public IPerforceSourceControlWorker
+class FPerforceRevertWorker final : public IPerforceSourceControlWorker
 {
 public:
-	virtual ~FPerforceRevertWorker() {}
+	FPerforceRevertWorker(FPerforceSourceControlProvider& InSourceControlProvider)
+		: IPerforceSourceControlWorker(InSourceControlProvider)
+	{}
+	virtual ~FPerforceRevertWorker() = default;
+
 	// IPerforceSourceControlWorker interface
 	virtual FName GetName() const override;
 	virtual bool Execute(class FPerforceSourceControlCommand& InCommand) override;
@@ -115,10 +143,14 @@ public:
 	FPerforceSourceControlChangelist ChangelistToUpdate;
 };
 
-class FPerforceSyncWorker : public IPerforceSourceControlWorker
+class FPerforceSyncWorker final : public IPerforceSourceControlWorker
 {
 public:
-	virtual ~FPerforceSyncWorker() {}
+	FPerforceSyncWorker(FPerforceSourceControlProvider& InSourceControlProvider)
+		: IPerforceSourceControlWorker(InSourceControlProvider)
+	{}
+	virtual ~FPerforceSyncWorker() = default;
+
 	// IPerforceSourceControlWorker interface
 	virtual FName GetName() const override;
 	virtual bool Execute(class FPerforceSourceControlCommand& InCommand) override;
@@ -129,10 +161,14 @@ public:
 	TMap<FString, EPerforceState::Type> OutResults;
 };
 
-class FPerforceUpdateStatusWorker : public IPerforceSourceControlWorker
+class FPerforceUpdateStatusWorker final : public IPerforceSourceControlWorker
 {
 public:
-	virtual ~FPerforceUpdateStatusWorker() {}
+	FPerforceUpdateStatusWorker(FPerforceSourceControlProvider& InSourceControlProvider)
+		: IPerforceSourceControlWorker(InSourceControlProvider)
+	{}
+	virtual ~FPerforceUpdateStatusWorker() = default;
+
 	// IPerforceSourceControlWorker interface
 	virtual FName GetName() const override;
 	virtual bool Execute(class FPerforceSourceControlCommand& InCommand) override;
@@ -155,20 +191,28 @@ public:
 	bool bForceQuiet = false;
 };
 
-class FPerforceGetWorkspacesWorker : public IPerforceSourceControlWorker
+class FPerforceGetWorkspacesWorker final : public IPerforceSourceControlWorker
 {
 public:
-	virtual ~FPerforceGetWorkspacesWorker() {}
+	FPerforceGetWorkspacesWorker(FPerforceSourceControlProvider& InSourceControlProvider)
+		: IPerforceSourceControlWorker(InSourceControlProvider)
+	{}
+	virtual ~FPerforceGetWorkspacesWorker() = default;
+
 	// IPerforceSourceControlWorker interface
 	virtual FName GetName() const override;
 	virtual bool Execute(class FPerforceSourceControlCommand& InCommand) override;
 	virtual bool UpdateStates() const override;
 };
 
-class FPerforceGetPendingChangelistsWorker : public IPerforceSourceControlWorker
+class FPerforceGetPendingChangelistsWorker final : public IPerforceSourceControlWorker
 {
 public:
-	virtual ~FPerforceGetPendingChangelistsWorker() {}
+	FPerforceGetPendingChangelistsWorker(FPerforceSourceControlProvider& InSourceControlProvider)
+		: IPerforceSourceControlWorker(InSourceControlProvider)
+	{}
+	virtual ~FPerforceGetPendingChangelistsWorker() = default;
+
 	// IPerforceSourceControlWorker interface
 	virtual FName GetName() const override;
 	virtual bool Execute(class FPerforceSourceControlCommand& InCommand) override;
@@ -186,10 +230,14 @@ private:
 	bool bCleanupCache = false;
 };
 
-class FPerforceCopyWorker : public IPerforceSourceControlWorker
+class FPerforceCopyWorker final : public IPerforceSourceControlWorker
 {
 public:
-	virtual ~FPerforceCopyWorker() {}
+	FPerforceCopyWorker(FPerforceSourceControlProvider& InSourceControlProvider)
+		: IPerforceSourceControlWorker(InSourceControlProvider)
+	{}
+	virtual ~FPerforceCopyWorker() = default;
+
 	// IPerforceSourceControlWorker interface
 	virtual FName GetName() const override;
 	virtual bool Execute(class FPerforceSourceControlCommand& InCommand) override;
@@ -200,10 +248,14 @@ public:
 	TMap<FString, EPerforceState::Type> OutResults;
 };
 
-class FPerforceResolveWorker : public IPerforceSourceControlWorker
+class FPerforceResolveWorker final : public IPerforceSourceControlWorker
 {
 public:
-	virtual ~FPerforceResolveWorker() {}
+	FPerforceResolveWorker(FPerforceSourceControlProvider& InSourceControlProvider)
+		: IPerforceSourceControlWorker(InSourceControlProvider)
+	{}
+
+	virtual ~FPerforceResolveWorker() = default;
 	// IPerforceSourceControlWorker interface
 	virtual FName GetName() const override;
 	virtual bool Execute(class FPerforceSourceControlCommand& InCommand) override;
@@ -212,21 +264,26 @@ private:
 	TArray< FString > UpdatedFiles;
 };
 
-class FPerforceChangeStatusWorker : public IPerforceSourceControlWorker
+class FPerforceChangeStatusWorker final : public IPerforceSourceControlWorker
 {
 public:
-	virtual ~FPerforceChangeStatusWorker() {}
+	FPerforceChangeStatusWorker(FPerforceSourceControlProvider& InSourceControlProvider)
+		: IPerforceSourceControlWorker(InSourceControlProvider)
+	{}
+	virtual ~FPerforceChangeStatusWorker() = default;
+
 	// IPerforceSourceControlWorker interface
 	virtual FName GetName() const override;
 	virtual bool Execute(class FPerforceSourceControlCommand& InCommand) override;
 	virtual bool UpdateStates() const override;
 };
 
-class FPerforceNewChangelistWorker : public IPerforceSourceControlWorker
+class FPerforceNewChangelistWorker final : public IPerforceSourceControlWorker
 {
 public:
-	FPerforceNewChangelistWorker();
-	virtual ~FPerforceNewChangelistWorker() {}
+	FPerforceNewChangelistWorker(FPerforceSourceControlProvider& InSourceControlProvider);
+	virtual ~FPerforceNewChangelistWorker() = default;
+
 	// IPerforceSourceControlWorker interface
 	virtual FName GetName() const override;
 	virtual bool Execute(class FPerforceSourceControlCommand& InCommand) override;
@@ -241,10 +298,14 @@ public:
 	TArray<FString> MovedFiles;
 };
 
-class FPerforceDeleteChangelistWorker : public IPerforceSourceControlWorker
+class FPerforceDeleteChangelistWorker final : public IPerforceSourceControlWorker
 {
 public:
-	virtual ~FPerforceDeleteChangelistWorker() {}
+	FPerforceDeleteChangelistWorker(FPerforceSourceControlProvider& InSourceControlProvider)
+		: IPerforceSourceControlWorker(InSourceControlProvider)
+	{}
+	virtual ~FPerforceDeleteChangelistWorker() = default;
+
 	// IPerforceSourceControlWorker interface
 	virtual FName GetName() const override;
 	virtual bool Execute(class FPerforceSourceControlCommand& InCommand) override;
@@ -254,10 +315,14 @@ public:
 	FPerforceSourceControlChangelist DeletedChangelist;
 };
 
-class FPerforceEditChangelistWorker : public IPerforceSourceControlWorker
+class FPerforceEditChangelistWorker final : public IPerforceSourceControlWorker
 {
 public:
-	virtual ~FPerforceEditChangelistWorker() {}
+	FPerforceEditChangelistWorker(FPerforceSourceControlProvider& InSourceControlProvider)
+		: IPerforceSourceControlWorker(InSourceControlProvider)
+	{}
+	virtual ~FPerforceEditChangelistWorker() = default;
+
 	// IPerforceSourceControlWorker interface
 	virtual FName GetName() const override;
 	virtual bool Execute(class FPerforceSourceControlCommand& InCommand) override;
@@ -268,10 +333,14 @@ public:
 	FText EditedDescription;
 };
 
-class FPerforceRevertUnchangedWorker : public IPerforceSourceControlWorker
+class FPerforceRevertUnchangedWorker final : public IPerforceSourceControlWorker
 {
 public:
-	virtual ~FPerforceRevertUnchangedWorker() {}
+	FPerforceRevertUnchangedWorker(FPerforceSourceControlProvider& InSourceControlProvider)
+		: IPerforceSourceControlWorker(InSourceControlProvider)
+	{}
+	virtual ~FPerforceRevertUnchangedWorker() = default;
+
 	// IPerforceSourceControlWorker interface
 	virtual FName GetName() const override;
 	virtual bool Execute(class FPerforceSourceControlCommand& InCommand) override;
@@ -284,10 +353,14 @@ protected:
 	FPerforceSourceControlChangelist ChangelistToUpdate;
 };
 
-class FPerforceReopenWorker : public IPerforceSourceControlWorker
+class FPerforceReopenWorker final : public IPerforceSourceControlWorker
 {
 public:
-	virtual ~FPerforceReopenWorker() {}
+	FPerforceReopenWorker(FPerforceSourceControlProvider& InSourceControlProvider)
+		: IPerforceSourceControlWorker(InSourceControlProvider)
+	{}
+	virtual ~FPerforceReopenWorker() = default;
+
 	// IPerforceSourceControlWorker interface
 	virtual FName GetName() const override;
 	virtual bool Execute(class FPerforceSourceControlCommand& InCommand) override;
@@ -301,10 +374,14 @@ protected:
 	FPerforceSourceControlChangelist DestinationChangelist;
 };
 
-class FPerforceShelveWorker : public IPerforceSourceControlWorker
+class FPerforceShelveWorker final : public IPerforceSourceControlWorker
 {
 public:
-	virtual ~FPerforceShelveWorker() {}
+	FPerforceShelveWorker(FPerforceSourceControlProvider& InSourceControlProvider)
+		: IPerforceSourceControlWorker(InSourceControlProvider)
+	{}
+	virtual ~FPerforceShelveWorker() = default;
+
 	// IPerforceSourceControlWorker interface
 	virtual FName GetName() const override;
 	virtual bool Execute(class FPerforceSourceControlCommand& InCommand) override;
@@ -328,10 +405,14 @@ protected:
 	FPerforceSourceControlChangelist OutChangelistToUpdate;
 };
 
-class FPerforceDeleteShelveWorker : public IPerforceSourceControlWorker
+class FPerforceDeleteShelveWorker final : public IPerforceSourceControlWorker
 {
 public:
-	virtual ~FPerforceDeleteShelveWorker() {}
+	FPerforceDeleteShelveWorker(FPerforceSourceControlProvider& InSourceControlProvider)
+		: IPerforceSourceControlWorker(InSourceControlProvider)
+	{}
+	virtual ~FPerforceDeleteShelveWorker() = default;
+
 	// IPerforceSourceControlWorker interface
 	virtual FName GetName() const override;
 	virtual bool Execute(class FPerforceSourceControlCommand& InCommand) override;
@@ -348,7 +429,11 @@ protected:
 class FPerforceUnshelveWorker : public IPerforceSourceControlWorker
 {
 public:
-	virtual ~FPerforceUnshelveWorker() {}
+	FPerforceUnshelveWorker(FPerforceSourceControlProvider& InSourceControlProvider)
+		: IPerforceSourceControlWorker(InSourceControlProvider)
+	{}
+	virtual ~FPerforceUnshelveWorker() = default;
+
 	// IPerforceSourceControlWorker interface
 	virtual FName GetName() const override;
 	virtual bool Execute(class FPerforceSourceControlCommand& InCommand) override;
@@ -362,9 +447,12 @@ protected:
 	TArray<FPerforceSourceControlState> ChangelistFilesStates;
 };
 
-class FPerforceDownloadFileWorker : public IPerforceSourceControlWorker
+class FPerforceDownloadFileWorker final : public IPerforceSourceControlWorker
 {
 public:
+	FPerforceDownloadFileWorker(FPerforceSourceControlProvider& InSourceControlProvider)
+		: IPerforceSourceControlWorker(InSourceControlProvider)
+	{}
 	virtual ~FPerforceDownloadFileWorker() = default;
 
 	// IPerforceSourceControlWorker interface
@@ -373,9 +461,12 @@ public:
 	virtual bool UpdateStates() const override;
 };
 
-class FPerforceCreateWorkspaceWorker : public IPerforceSourceControlWorker
+class FPerforceCreateWorkspaceWorker final : public IPerforceSourceControlWorker
 {
 public:
+	FPerforceCreateWorkspaceWorker(FPerforceSourceControlProvider& InSourceControlProvider)
+		: IPerforceSourceControlWorker(InSourceControlProvider)
+	{}
 	virtual ~FPerforceCreateWorkspaceWorker() = default;
 
 	// IPerforceSourceControlWorker interface
@@ -387,9 +478,12 @@ private:
 	void AddType(const class FCreateWorkspace& Operation, FStringBuilderBase& ClientDesc);
 };
 
-class FPerforceDeleteWorkspaceWorker : public IPerforceSourceControlWorker
+class FPerforceDeleteWorkspaceWorker final : public IPerforceSourceControlWorker
 {
 public:
+	FPerforceDeleteWorkspaceWorker(FPerforceSourceControlProvider& InSourceControlProvider)
+		: IPerforceSourceControlWorker(InSourceControlProvider)
+	{}
 	virtual ~FPerforceDeleteWorkspaceWorker() = default;
 
 	// IPerforceSourceControlWorker interface
diff --git a/Engine/Plugins/Developer/PerforceSourceControl/Source/PerforceSourceControl/Private/PerforceSourceControlProvider.cpp b/Engine/Plugins/Developer/PerforceSourceControl/Source/PerforceSourceControl/Private/PerforceSourceControlProvider.cpp
--- a/Engine/Plugins/Developer/PerforceSourceControl/Source/PerforceSourceControl/Private/PerforceSourceControlProvider.cpp
+++ b/Engine/Plugins/Developer/PerforceSourceControl/Source/PerforceSourceControl/Private/PerforceSourceControlProvider.cpp
@@ -1,38 +1,73 @@
 // Copyright Epic Games, Inc. All Rights Reserved.
 
 #include "PerforceSourceControlProvider.h"
-#include "PerforceSourceControlPrivate.h"
+
+#include "Algo/Transform.h"
 #include "HAL/PlatformProcess.h"
-#include "Misc/MessageDialog.h"
+#include "ISourceControlLabel.h"
+#include "ISourceControlModule.h"
+#include "Logging/MessageLog.h"
 #include "Misc/CommandLine.h"
-#include "Algo/Transform.h"
+#include "Misc/MessageDialog.h"
 #include "Misc/QueuedThreadPool.h"
-#include "Modules/ModuleManager.h"
-#include "Widgets/DeclarativeSyntaxSupport.h"
-#include "ISourceControlModule.h"
+#include "PerforceConnection.h"
 #include "PerforceConnectionInfo.h"
 #include "PerforceSourceControlCommand.h"
-#include "ISourceControlLabel.h"
 #include "PerforceSourceControlLabel.h"
-#include "PerforceConnection.h"
-#include "PerforceSourceControlModule.h"
+#include "PerforceSourceControlPrivate.h"
 #include "SPerforceSourceControlSettings.h"
-#include "Logging/MessageLog.h"
 #include "ScopedSourceControlProgress.h"
 #include "SourceControlHelpers.h"
 #include "SourceControlOperations.h"
+#include "Widgets/DeclarativeSyntaxSupport.h"
 
 static FName ProviderName("Perforce");
 
 #define LOCTEXT_NAMESPACE "PerforceSourceControl"
 
-/** Init of connection with source control server */
+FPerforceSourceControlProvider::FPerforceSourceControlProvider()
+	: PerforceSCCSettings(*this, FStringView())
+	, InitialSettings(FSourceControlInitSettings::EBehavior::OverrideAll)
+	, OwnerName(TEXT("Default"))
+	, bServerAvailable(false)
+	, bLoginError(false)
+	, PersistentConnection(nullptr)
+{
+	auto ParseCmdLineSetting = [this](const TCHAR* SettingKey) -> void
+	{
+		FString SettingValue;
+		if (FParse::Value(FCommandLine::Get(), SettingKey, SettingValue))
+		{
+			InitialSettings.AddSetting(SettingKey, SettingValue);
+		}
+	};
+
+	ParseCmdLineSetting(TEXT("P4Port="));
+	ParseCmdLineSetting(TEXT("P4User="));
+	ParseCmdLineSetting(TEXT("P4Client="));
+	ParseCmdLineSetting(TEXT("P4Host="));
+	ParseCmdLineSetting(TEXT("P4Passwd="));
+	ParseCmdLineSetting(TEXT("P4Changelist="));
+
+	AccessSettings().LoadSettings();
+}
+
+FPerforceSourceControlProvider::FPerforceSourceControlProvider(const FStringView& InOwnerName, const FSourceControlInitSettings& InInitialSettings)
+	: PerforceSCCSettings(*this, InOwnerName)
+	, InitialSettings(InInitialSettings)
+	, OwnerName(InOwnerName)
+	, bServerAvailable(false)
+	, bLoginError(false)
+	, PersistentConnection(nullptr)
+{
+	AccessSettings().LoadSettings();
+}
+
 void FPerforceSourceControlProvider::Init(bool bForceConnection)
 {
 	ParseCommandLineSettings(bForceConnection);
 }
 
-/** API Specific close the connection with source control server*/
 void FPerforceSourceControlProvider::Close()
 {
 	if ( PersistentConnection )
@@ -84,8 +119,7 @@ TSharedRef<FPerforceSourceControlChangelistState, ESPMode::ThreadSafe> FPerforce
 
 FText FPerforceSourceControlProvider::GetStatusText() const
 {
-	FPerforceSourceControlModule& PerforceSourceControl = FModuleManager::LoadModuleChecked<FPerforceSourceControlModule>( "PerforceSourceControl" );
-	const FPerforceSourceControlSettings& Settings = PerforceSourceControl.AccessSettings();
+	const FPerforceSourceControlSettings& Settings = AccessSettings();
 
 	FFormatNamedArguments Args;
 	Args.Add( TEXT("IsEnabled"), IsEnabled() ? LOCTEXT("Yes", "Yes") : LOCTEXT("No", "No") );
@@ -121,20 +155,19 @@ bool FPerforceSourceControlProvider::IsAvailable() const
 
 bool FPerforceSourceControlProvider::EstablishPersistentConnection()
 {
-	FPerforceSourceControlModule& PerforceSourceControl = FModuleManager::LoadModuleChecked<FPerforceSourceControlModule>( "PerforceSourceControl" );
-	FPerforceConnectionInfo ConnectionInfo = PerforceSourceControl.AccessSettings().GetConnectionInfo();
+	FPerforceConnectionInfo ConnectionInfo = AccessSettings().GetConnectionInfo();
 
 	bool bIsValidConnection = false;
 	if ( !PersistentConnection )
 	{
-		PersistentConnection = new FPerforceConnection(ConnectionInfo);
+		PersistentConnection = new FPerforceConnection(ConnectionInfo, *this);
 	}
 
 	bIsValidConnection = PersistentConnection->IsValidConnection();
 	if ( !bIsValidConnection )
 	{
 		delete PersistentConnection;
-		PersistentConnection = new FPerforceConnection(ConnectionInfo);
+		PersistentConnection = new FPerforceConnection(ConnectionInfo, *this);
 		bIsValidConnection = PersistentConnection->IsValidConnection();
 	}
 
@@ -142,32 +175,38 @@ bool FPerforceSourceControlProvider::EstablishPersistentConnection()
 	return bIsValidConnection;
 }
 
-/**
- * Loads user/SCC information from the command line or INI file.
- */
 void FPerforceSourceControlProvider::ParseCommandLineSettings(bool bForceConnection)
 {
-	ISourceControlModule& SourceControlModule = FModuleManager::LoadModuleChecked<ISourceControlModule>( "SourceControl" );
-	FPerforceSourceControlModule& PerforceSourceControl = FModuleManager::GetModuleChecked<FPerforceSourceControlModule>( "PerforceSourceControl" );
-
-	bool bFoundCmdLineSettings = false;
-
-	// Check command line for any overridden settings
-	FPerforceSourceControlSettings& P4Settings = PerforceSourceControl.AccessSettings();
+	FPerforceSourceControlSettings& P4Settings = AccessSettings();
 
+	// First we take a copy of the existing settings
 	FString PortName = P4Settings.GetPort();
 	FString UserName = P4Settings.GetUserName();
 	FString ClientSpecName = P4Settings.GetWorkspace();
 	FString HostOverrideName = P4Settings.GetHostOverride();
 	FString Changelist = P4Settings.GetChangelistNumber();
-	bFoundCmdLineSettings = FParse::Value(FCommandLine::Get(), TEXT("P4Port="), PortName);
-	bFoundCmdLineSettings |= FParse::Value(FCommandLine::Get(), TEXT("P4User="), UserName);
-	bFoundCmdLineSettings |= FParse::Value(FCommandLine::Get(), TEXT("P4Client="), ClientSpecName);
-	bFoundCmdLineSettings |= FParse::Value(FCommandLine::Get(), TEXT("P4Host="), HostOverrideName);
-	bFoundCmdLineSettings |= FParse::Value(FCommandLine::Get(), TEXT("P4Passwd="), Ticket);
-	bFoundCmdLineSettings |= FParse::Value(FCommandLine::Get(), TEXT("P4Changelist="), Changelist);
-	if(bFoundCmdLineSettings)
+
+	EConnectionOptions Options = EConnectionOptions::None;
+
+	// Then we see if any of these settings are overridden by the initial settings
+	// Note that as long as one setting is overridden, we will reset all non-overridden 
+	// values to empty and later calculate them via FPerforceConnection::EnsureValidConnection
+	if (InitialSettings.HasOverrides())
 	{
+		InitialSettings.OverrideSetting(TEXT("P4Port"),			PortName);
+		InitialSettings.OverrideSetting(TEXT("P4User"),			UserName);
+		InitialSettings.OverrideSetting(TEXT("P4Client"),		ClientSpecName);
+		InitialSettings.OverrideSetting(TEXT("P4Host"),			HostOverrideName);
+		InitialSettings.OverrideSetting(TEXT("P4Passwd"),		Ticket);
+		InitialSettings.OverrideSetting(TEXT("P4Changelist"),	Changelist);
+
+		// If P4Client is overridden then we can rely on that, even if it is blank (no workspace) and 
+		// so don't need to automatically find a workspace when ensuring the connection.
+		if (InitialSettings.IsOverridden(TEXT("P4Client")))
+		{
+			Options |= EConnectionOptions::WorkspaceOptional;
+		}
+
 		P4Settings.SetPort(PortName);
 		P4Settings.SetUserName(UserName);
 		P4Settings.SetWorkspace(ClientSpecName);
@@ -179,7 +218,7 @@ void FPerforceSourceControlProvider::ParseCommandLineSettings(bool bForceConnect
 	{
 		bLoginError = false;
 		FPerforceConnectionInfo ConnectionInfo = P4Settings.GetConnectionInfo();
-		if(FPerforceConnection::EnsureValidConnection(PortName, UserName, ClientSpecName, ConnectionInfo))
+		if(FPerforceConnection::EnsureValidConnection(PortName, UserName, ClientSpecName, ConnectionInfo, *this, Options))
 		{
 			P4Settings.SetPort(PortName);
 			P4Settings.SetUserName(UserName);
@@ -190,14 +229,13 @@ void FPerforceSourceControlProvider::ParseCommandLineSettings(bool bForceConnect
 	}
 
 	//Save off settings so this doesn't happen every time
-	PerforceSourceControl.SaveSettings();
+	SaveConnectionSettings();
 }
 
-
 void FPerforceSourceControlProvider::GetWorkspaceList(const FPerforceConnectionInfo& InConnectionInfo, TArray<FString>& OutWorkspaceList, TArray<FText>& OutErrorMessages)
 {
 	//attempt to ask perforce for a list of client specs that belong to this user
-	FPerforceConnection Connection(InConnectionInfo);
+	FPerforceConnection Connection(InConnectionInfo, *this);
 	Connection.GetWorkspaceList(InConnectionInfo, FOnIsCancelled(), OutWorkspaceList, OutErrorMessages);
 }
 
@@ -206,6 +244,11 @@ const FString& FPerforceSourceControlProvider::GetTicket() const
 	return Ticket;
 }
 
+const FString& FPerforceSourceControlProvider::GetOwnerName() const
+{
+	return OwnerName;
+}
+
 const FName& FPerforceSourceControlProvider::GetName() const
 {
 	return ProviderName;
@@ -534,7 +577,7 @@ void FPerforceSourceControlProvider::Tick()
 	}
 }
 
-static void ParseGetLabelsResults(const FP4RecordSet& InRecords, TArray< TSharedRef<ISourceControlLabel> >& OutLabels)
+static void ParseGetLabelsResults(FPerforceSourceControlProvider& InSourceControlProvider, const FP4RecordSet& InRecords, TArray< TSharedRef<ISourceControlLabel> >& OutLabels)
 {
 	// Iterate over each record found as a result of the command, parsing it for relevant information
 	for (int32 Index = 0; Index < InRecords.Num(); ++Index)
@@ -543,17 +586,17 @@ static void ParseGetLabelsResults(const FP4RecordSet& InRecords, TArray< TShared
 		FString LabelName = ClientRecord(TEXT("label"));
 		if(LabelName.Len() > 0)
 		{
-			OutLabels.Add(MakeShareable( new FPerforceSourceControlLabel(LabelName) ) );
+			OutLabels.Add(MakeShareable( new FPerforceSourceControlLabel(InSourceControlProvider, LabelName) ) );
 		}
 	}
 }
 
-TArray< TSharedRef<ISourceControlLabel> > FPerforceSourceControlProvider::GetLabels( const FString& InMatchingSpec ) const
+TArray< TSharedRef<ISourceControlLabel> > FPerforceSourceControlProvider::GetLabels(const FString& InMatchingSpec) const
 {
 	TArray< TSharedRef<ISourceControlLabel> > Labels;
 
-	FPerforceSourceControlModule& PerforceSourceControl = FModuleManager::LoadModuleChecked<FPerforceSourceControlModule>("PerforceSourceControl");
-	FScopedPerforceConnection ScopedConnection(EConcurrency::Synchronous, PerforceSourceControl.AccessSettings().GetConnectionInfo());
+	// const_cast to avoid changing the ISourceControlProvider API as it is hard to deprecate without causing derived types to give compiler errors.
+	FScopedPerforceConnection ScopedConnection(EConcurrency::Synchronous, *const_cast<FPerforceSourceControlProvider*>(this));
 	if(ScopedConnection.IsValid())
 	{
 		FPerforceConnection& Connection = ScopedConnection.GetConnection();
@@ -565,7 +608,8 @@ TArray< TSharedRef<ISourceControlLabel> > FPerforceSourceControlProvider::GetLab
 		bool bConnectionDropped = false;
 		if(Connection.RunCommand(TEXT("labels"), Parameters, Records, ErrorMessages, FOnIsCancelled(), bConnectionDropped))
 		{
-			ParseGetLabelsResults(Records, Labels);
+			// const_cast to avoid changing the ISourceControlProvider API as it is hard to deprecate without causing derived types to give compiler errors.
+			ParseGetLabelsResults(*const_cast<FPerforceSourceControlProvider*>(this), Records, Labels);
 		}
 		else
 		{
@@ -668,16 +712,15 @@ ECommandResult::Type FPerforceSourceControlProvider::SwitchWorkspace(FStringView
 	// Do not call Init directly as we do not want to save the new workspace name to
 	// the source control settings!
 
-	FPerforceSourceControlModule& PerforceSourceControl = FModuleManager::LoadModuleChecked<FPerforceSourceControlModule>("PerforceSourceControl");
-	const FString OldWorkspaceName = PerforceSourceControl.AccessSettings().GetWorkspace();
+	FPerforceSourceControlSettings& P4Settings = AccessSettings(); 
+	
+	const FString OldWorkspaceName = P4Settings.GetWorkspace();
 		
-	FPerforceSourceControlSettings& P4Settings = PerforceSourceControl.AccessSettings();
-
 	FString PortName = P4Settings.GetPort();
 	FString UserName = P4Settings.GetUserName();
 	FString ClientSpecName = FString(NewWorkspaceName);
 
-	if (FPerforceConnection::EnsureValidConnection(PortName, UserName, ClientSpecName, P4Settings.GetConnectionInfo()))
+	if (FPerforceConnection::EnsureValidConnection(PortName, UserName, ClientSpecName, P4Settings.GetConnectionInfo(), *this, EConnectionOptions::None))
 	{
 		P4Settings.SetPort(PortName);
 		P4Settings.SetUserName(UserName);
@@ -702,7 +745,7 @@ ECommandResult::Type FPerforceSourceControlProvider::SwitchWorkspace(FStringView
 		OutResultInfo.ErrorMessages.Add(Message);
 
 		// The connection didn't work so we should try to restore the old workspace name
-		PerforceSourceControl.AccessSettings().SetWorkspace(OldWorkspaceName);
+		P4Settings.SetWorkspace(OldWorkspaceName);
 		
 		return ECommandResult::Failed;
 	}
@@ -711,24 +754,35 @@ ECommandResult::Type FPerforceSourceControlProvider::SwitchWorkspace(FStringView
 #if SOURCE_CONTROL_WITH_SLATE
 TSharedRef<class SWidget> FPerforceSourceControlProvider::MakeSettingsWidget() const
 {
-	return SNew(SPerforceSourceControlSettings);
+	// const_cast to avoid changing the ISourceControlProvider API as it is hard to deprecate without
+	// causing derived types to give compiler errors.
+	return SNew(SPerforceSourceControlSettings, const_cast<FPerforceSourceControlProvider*>(this));
 }
 #endif
 
-TSharedPtr<IPerforceSourceControlWorker, ESPMode::ThreadSafe> FPerforceSourceControlProvider::CreateWorker(const FName& InOperationName) const
+TUniquePtr<ISourceControlProvider> FPerforceSourceControlProvider::Create(const FStringView& InOwnerName, const FSourceControlInitSettings& InInitialSettings) const
+{ 
+	return MakeUnique<FPerforceSourceControlProvider>(InOwnerName, InInitialSettings);
+}
+
+const FPerforceSourceControlSettings& FPerforceSourceControlProvider::AccessSettings() const
 {
-	const FGetPerforceSourceControlWorker* Operation = WorkersMap.Find(InOperationName);
-	if(Operation != NULL)
-	{
-		return Operation->Execute();
-	}
+	return PerforceSCCSettings;
+}
+
+FPerforceSourceControlSettings& FPerforceSourceControlProvider::AccessSettings()
+{
+	return PerforceSCCSettings;
+}
 
-	return NULL;
+void FPerforceSourceControlProvider::SaveConnectionSettings()
+{
+	PerforceSCCSettings.SaveSettings();
 }
 
-void FPerforceSourceControlProvider::RegisterWorker( const FName& InName, const FGetPerforceSourceControlWorker& InDelegate )
+TSharedPtr<IPerforceSourceControlWorker, ESPMode::ThreadSafe> FPerforceSourceControlProvider::CreateWorker(const FName& InOperationName)
 {
-	WorkersMap.Add( InName, InDelegate );
+	return IPerforceSourceControlWorker::CreateWorker(InOperationName, *this);
 }
 
 ECommandResult::Type FPerforceSourceControlProvider::ExecuteSynchronousCommand(FPerforceSourceControlCommand& InCommand, const FText& Task, bool bSuppressResponseMsg)
@@ -823,8 +877,7 @@ bool FPerforceSourceControlProvider::QueryStateBranchConfig(const FString& Confi
 	}
 
 	// Request branch configuration from depot
-	FPerforceSourceControlModule& PerforceSourceControl = FModuleManager::LoadModuleChecked<FPerforceSourceControlModule>("PerforceSourceControl");
-	FScopedPerforceConnection ScopedConnection(EConcurrency::Synchronous, PerforceSourceControl.AccessSettings().GetConnectionInfo());
+	FScopedPerforceConnection ScopedConnection(EConcurrency::Synchronous, *this);
 	if (ScopedConnection.IsValid())
 	{
 		FPerforceConnection& Connection = ScopedConnection.GetConnection();
diff --git a/Engine/Plugins/Developer/PerforceSourceControl/Source/PerforceSourceControl/Private/PerforceSourceControlProvider.h b/Engine/Plugins/Developer/PerforceSourceControl/Source/PerforceSourceControl/Private/PerforceSourceControlProvider.h
--- a/Engine/Plugins/Developer/PerforceSourceControl/Source/PerforceSourceControl/Private/PerforceSourceControlProvider.h
+++ b/Engine/Plugins/Developer/PerforceSourceControl/Source/PerforceSourceControl/Private/PerforceSourceControlProvider.h
@@ -3,27 +3,24 @@
 #pragma once
 
 #include "CoreMinimal.h"
+#include "IPerforceSourceControlWorker.h"
 #include "ISourceControlOperation.h"
-#include "ISourceControlState.h"
 #include "ISourceControlProvider.h"
-#include "IPerforceSourceControlWorker.h"
-#include "PerforceSourceControlState.h"
+#include "ISourceControlState.h"
 #include "PerforceSourceControlChangelistState.h"
+#include "PerforceSourceControlSettings.h"
+#include "PerforceSourceControlState.h"
+#include "SourceControlInitSettings.h"
 
 class FPerforceSourceControlCommand;
 
-DECLARE_DELEGATE_RetVal(FPerforceSourceControlWorkerRef, FGetPerforceSourceControlWorker)
-
 class FPerforceSourceControlProvider : public ISourceControlProvider
 {
 public:
-	/** Constructor */
-	FPerforceSourceControlProvider()
-		: bServerAvailable(false)
-		, bLoginError(false)
-		, PersistentConnection(NULL)
-	{
-	}
+	
+	FPerforceSourceControlProvider();
+	FPerforceSourceControlProvider(const FStringView& OwnerName, const FSourceControlInitSettings& InitialSettings);
+	virtual ~FPerforceSourceControlProvider() = default;
 
 	/* ISourceControlProvider implementation */
 	virtual void Init(bool bForceConnection = true) override;
@@ -59,12 +56,6 @@ public:
 
 	using ISourceControlProvider::Execute;
 
-	/**
-	 * Register a worker with the provider.
-	 * This is used internally so the provider can maintain a map of all available operations.
-	 */
-	void RegisterWorker( const FName& InName, const FGetPerforceSourceControlWorker& InDelegate );
-
 	/**
 	 * Gets a list of client spec names from the source control provider
 	 *
@@ -77,6 +68,9 @@ public:
 	/** Get the P4 ticket we will use for connections */
 	const FString& GetTicket() const;
 
+	/** Returns the name of the system that owns the provider */
+	const FString& GetOwnerName() const;
+
 	/** Set list of error messages that occurred after last perforce command */
 	void SetLastErrors(const TArray<FText>& InErrors);
 
@@ -117,10 +111,20 @@ public:
 	/** Returns a list of changelists from the cache based on a given predicate */
 	TArray<FSourceControlChangelistStateRef> GetCachedStateByPredicate(TFunctionRef<bool(const FSourceControlChangelistStateRef&)> Predicate) const;
 
+	/** Returns the settings for the current source control provider */
+	const FPerforceSourceControlSettings& AccessSettings() const;
+
+	/** Returns the settings for the current source control provider */
+	FPerforceSourceControlSettings& AccessSettings();
+
 private:
+	/** Instantiates and returns a new FPerforceSourceControlProvider */
+	virtual TUniquePtr<ISourceControlProvider> Create(const FStringView& OwnerName, const FSourceControlInitSettings& InitialSettings) const override;
+
+	void SaveConnectionSettings();
 
 	/** Helper function used to create a worker for a particular operation */
-	TSharedPtr<class IPerforceSourceControlWorker, ESPMode::ThreadSafe> CreateWorker(const FName& InOperationName) const;
+	TSharedPtr<class IPerforceSourceControlWorker, ESPMode::ThreadSafe> CreateWorker(const FName& InOperationName);
 
 	/** 
 	 * Logs any messages that a command needs to output.
@@ -146,6 +150,15 @@ private:
 
 private:
 
+	/** The settings for Perforce source control */
+	FPerforceSourceControlSettings PerforceSCCSettings;
+
+	/** The initial settings for the provider. These are used every time ::Init is called */
+	FSourceControlInitSettings InitialSettings;
+
+	/** Name of the system that owns the provider */
+	FString OwnerName;
+
 	/** The ticket we use for login. */
 	FString Ticket;
 
@@ -171,9 +184,6 @@ private:
 	TMap<FString, TSharedRef<class FPerforceSourceControlState, ESPMode::ThreadSafe> > StateCache;
 	TMap<FPerforceSourceControlChangelist, TSharedRef<class FPerforceSourceControlChangelistState, ESPMode::ThreadSafe> > ChangelistsStateCache;
 
-	/** The currently registered source control operations */
-	TMap<FName, FGetPerforceSourceControlWorker> WorkersMap;
-
 	/** Queue for commands given by the main thread */
 	TArray < FPerforceSourceControlCommand* > CommandQueue;
 
diff --git a/Engine/Plugins/Developer/PerforceSourceControl/Source/PerforceSourceControl/Private/PerforceSourceControlRevision.cpp b/Engine/Plugins/Developer/PerforceSourceControl/Source/PerforceSourceControl/Private/PerforceSourceControlRevision.cpp
--- a/Engine/Plugins/Developer/PerforceSourceControl/Source/PerforceSourceControl/Private/PerforceSourceControlRevision.cpp
+++ b/Engine/Plugins/Developer/PerforceSourceControl/Source/PerforceSourceControl/Private/PerforceSourceControlRevision.cpp
@@ -1,21 +1,30 @@
 // Copyright Epic Games, Inc. All Rights Reserved.
 
 #include "PerforceSourceControlRevision.h"
-#include "PerforceSourceControlPrivate.h"
+
+#include "Logging/MessageLog.h"
 #include "Misc/FileHelper.h"
 #include "Misc/Paths.h"
-#include "Modules/ModuleManager.h"
-#include "PerforceSourceControlModule.h"
 #include "PerforceConnection.h"
-#include "Logging/MessageLog.h"
+#include "PerforceSourceControlProvider.h"
 
 #define LOCTEXT_NAMESPACE "PerforceSourceControl"
 
+FPerforceSourceControlRevision:: FPerforceSourceControlRevision(FPerforceSourceControlProvider& InSCCProvider)
+	: RevisionNumber(0)
+	, Date(0)
+	, ChangelistNumber(0)
+	, FileSize(0)
+	, bIsShelve(false)
+	, SCCProvider(InSCCProvider)
+{
+}
+
 bool FPerforceSourceControlRevision::Get( FString& InOutFilename, EConcurrency::Type InConcurrency) const
 {
 	bool bCommandOK = false;
-	FPerforceSourceControlModule& PerforceSourceControl = FModuleManager::LoadModuleChecked<FPerforceSourceControlModule>("PerforceSourceControl");
-	FScopedPerforceConnection ScopedConnection(InConcurrency, PerforceSourceControl.AccessSettings().GetConnectionInfo());
+	
+	FScopedPerforceConnection ScopedConnection(InConcurrency, GetSCCProvider());
 	if(ScopedConnection.IsValid())
 	{
 		FPerforceConnection& Connection = ScopedConnection.GetConnection();
@@ -148,8 +157,8 @@ static void ParseAnnotationResults(const FP4RecordSet& Records, TArray<FAnnotati
 bool FPerforceSourceControlRevision::GetAnnotated( TArray<FAnnotationLine>& OutLines ) const
 {
 	bool bCommandOK = false;
-	FPerforceSourceControlModule& PerforceSourceControl = FModuleManager::LoadModuleChecked<FPerforceSourceControlModule>("PerforceSourceControl");
-	FScopedPerforceConnection ScopedConnection(EConcurrency::Synchronous, PerforceSourceControl.AccessSettings().GetConnectionInfo());
+	
+	FScopedPerforceConnection ScopedConnection(EConcurrency::Synchronous, GetSCCProvider());
 	if(ScopedConnection.IsValid())
 	{
 		FPerforceConnection& Connection = ScopedConnection.GetConnection();
diff --git a/Engine/Plugins/Developer/PerforceSourceControl/Source/PerforceSourceControl/Private/PerforceSourceControlRevision.h b/Engine/Plugins/Developer/PerforceSourceControl/Source/PerforceSourceControl/Private/PerforceSourceControlRevision.h
--- a/Engine/Plugins/Developer/PerforceSourceControl/Source/PerforceSourceControl/Private/PerforceSourceControlRevision.h
+++ b/Engine/Plugins/Developer/PerforceSourceControl/Source/PerforceSourceControl/Private/PerforceSourceControlRevision.h
@@ -6,17 +6,14 @@
 #include "ISourceControlProvider.h"
 #include "ISourceControlRevision.h"
 
+class FPerforceSourceControlProvider;
+
 class FPerforceSourceControlRevision : public ISourceControlRevision, public TSharedFromThis<FPerforceSourceControlRevision, ESPMode::ThreadSafe>
 {
 public:
-	FPerforceSourceControlRevision()
-		: RevisionNumber(0)
-		, Date(0)
-		, ChangelistNumber(0)
-		, FileSize(0)
-		, bIsShelve(false)
-	{
-	}
+
+	FPerforceSourceControlRevision(FPerforceSourceControlProvider& InSCCProvider);
+	virtual ~FPerforceSourceControlRevision() = default;
 
 	/** ISourceControlRevision interface */
 	virtual bool Get( FString& InOutFilename, EConcurrency::Type InConcurrency = EConcurrency::Synchronous) const override;
@@ -35,6 +32,7 @@ public:
 	virtual int32 GetFileSize() const override;
 
 public:
+
 	/** The local filename the this revision refers to */
 	FString FileName;
 
@@ -70,4 +68,14 @@ public:
 
 	/** Whether this reprensents a revision bound to a shelved file in a changelist */
 	bool bIsShelve;
+
+private:
+	/** Internal accessor to the source control provider associated with the object */
+	FPerforceSourceControlProvider& GetSCCProvider() const
+	{
+		return SCCProvider;
+	}
+
+	/** The source control provider that this object is associated with */
+	FPerforceSourceControlProvider& SCCProvider;
 };
diff --git a/Engine/Plugins/Developer/PerforceSourceControl/Source/PerforceSourceControl/Private/PerforceSourceControlSettings.cpp b/Engine/Plugins/Developer/PerforceSourceControl/Source/PerforceSourceControl/Private/PerforceSourceControlSettings.cpp
--- a/Engine/Plugins/Developer/PerforceSourceControl/Source/PerforceSourceControl/Private/PerforceSourceControlSettings.cpp
+++ b/Engine/Plugins/Developer/PerforceSourceControl/Source/PerforceSourceControl/Private/PerforceSourceControlSettings.cpp
@@ -2,24 +2,27 @@
 
 #include "PerforceSourceControlSettings.h"
 
-#include "PerforceSourceControlPrivate.h"
-
-#include "Misc/ScopeLock.h"
+#include "Misc/App.h"
 #include "Misc/ConfigCacheIni.h"
-#include "Modules/ModuleManager.h"
+#include "Misc/ScopeLock.h"
+#include "PerforceSourceControlPrivate.h"
+#include "PerforceSourceControlProvider.h"
 #include "SPerforceSourceControlSettings.h"
-#include "PerforceSourceControlModule.h"
 #include "SourceControlHelpers.h"
 
-namespace PerforceSettingsConstants
+FPerforceSourceControlSettings::FPerforceSourceControlSettings(const FPerforceSourceControlProvider& InSCCProvider, const FStringView& OwnerName)
+	: SCCProvider(InSCCProvider)
 {
-
-/** The section of the ini file we load our settings from */
-static const FString SettingsSection = TEXT("PerforceSourceControl.PerforceSourceControlSettings");
-
+	if (OwnerName.IsEmpty())
+	{
+		SettingsSection = TEXT("PerforceSourceControl.PerforceSourceControlSettings");
+	}
+	else
+	{
+		SettingsSection = WriteToString<128>(TEXT("PerforceSourceControl."), OwnerName, TEXT("Settings"));
+	}
 }
 
-
 const FString& FPerforceSourceControlSettings::GetPort() const
 {
 	FScopeLock ScopeLock(&CriticalSection);
@@ -109,34 +112,39 @@ void FPerforceSourceControlSettings::LoadSettings()
 	FScopeLock ScopeLock(&CriticalSection);
 	const FString& IniFile = SourceControlHelpers::GetSettingsIni();
 
-	GConfig->GetBool(*PerforceSettingsConstants::SettingsSection, TEXT("UseP4Config"), ConnectionInfo.bUseP4Config, IniFile);
+	GConfig->GetBool(*SettingsSection, TEXT("UseP4Config"), ConnectionInfo.bUseP4Config, IniFile);
 	if (ConnectionInfo.bUseP4Config)
 	{
 		ImportP4Config();
 	}
 	else
 	{
-		if(!GConfig->GetString(*PerforceSettingsConstants::SettingsSection, TEXT("Port"), ConnectionInfo.Port, IniFile))
+		if(!GConfig->GetString(*SettingsSection, TEXT("Port"), ConnectionInfo.Port, IniFile))
 		{
 			// backwards compatibility - previously we mis-specified the Port as 'Host'
-			GConfig->GetString(*PerforceSettingsConstants::SettingsSection, TEXT("Host"), ConnectionInfo.Port, IniFile);
+			GConfig->GetString(*SettingsSection, TEXT("Host"), ConnectionInfo.Port, IniFile);
 		}
-		GConfig->GetString(*PerforceSettingsConstants::SettingsSection, TEXT("UserName"), ConnectionInfo.UserName, IniFile);
-		GConfig->GetString(*PerforceSettingsConstants::SettingsSection, TEXT("Workspace"), ConnectionInfo.Workspace, IniFile);
-	}	
-	GConfig->GetString(*PerforceSettingsConstants::SettingsSection, TEXT("HostOverride"), ConnectionInfo.HostOverride, IniFile);
+		GConfig->GetString(*SettingsSection, TEXT("UserName"), ConnectionInfo.UserName, IniFile);
+		GConfig->GetString(*SettingsSection, TEXT("Workspace"), ConnectionInfo.Workspace, IniFile);
+	}
+	GConfig->GetString(*SettingsSection, TEXT("HostOverride"), ConnectionInfo.HostOverride, IniFile);
 }
 
 void FPerforceSourceControlSettings::SaveSettings() const
 {
+	if (FApp::IsUnattended() || IsRunningCommandlet())
+	{
+		return;
+	}
+
 	FScopeLock ScopeLock(&CriticalSection);
 	const FString& IniFile = SourceControlHelpers::GetSettingsIni();
 	
-	GConfig->SetBool(*PerforceSettingsConstants::SettingsSection, TEXT("UseP4Config"), ConnectionInfo.bUseP4Config, IniFile);
-	GConfig->SetString(*PerforceSettingsConstants::SettingsSection, TEXT("Port"), *ConnectionInfo.Port, IniFile);
-	GConfig->SetString(*PerforceSettingsConstants::SettingsSection, TEXT("UserName"), *ConnectionInfo.UserName, IniFile);
-	GConfig->SetString(*PerforceSettingsConstants::SettingsSection, TEXT("Workspace"), *ConnectionInfo.Workspace, IniFile);
-	GConfig->SetString(*PerforceSettingsConstants::SettingsSection, TEXT("HostOverride"), *ConnectionInfo.HostOverride, IniFile);
+	GConfig->SetBool(*SettingsSection, TEXT("UseP4Config"), ConnectionInfo.bUseP4Config, IniFile);
+	GConfig->SetString(*SettingsSection, TEXT("Port"), *ConnectionInfo.Port, IniFile);
+	GConfig->SetString(*SettingsSection, TEXT("UserName"), *ConnectionInfo.UserName, IniFile);
+	GConfig->SetString(*SettingsSection, TEXT("Workspace"), *ConnectionInfo.Workspace, IniFile);
+	GConfig->SetString(*SettingsSection, TEXT("HostOverride"), *ConnectionInfo.HostOverride, IniFile);
 }
 
 FPerforceConnectionInfo FPerforceSourceControlSettings::GetConnectionInfo() const
@@ -155,9 +163,7 @@ FPerforceConnectionInfo FPerforceSourceControlSettings::GetConnectionInfo() cons
 	}
 
 	// Ticket is stored in the provider (this is only set by the command line so should be safe to access without threading protection)
-	FPerforceSourceControlModule& PerforceSourceControl = FModuleManager::GetModuleChecked<FPerforceSourceControlModule>( "PerforceSourceControl" );
-	const FPerforceSourceControlProvider& Provider = PerforceSourceControl.GetProvider();
-	OutConnectionInfo.Ticket = Provider.GetTicket();
+	OutConnectionInfo.Ticket = GetSCCProvider().GetTicket();
 
 	return OutConnectionInfo;
 }
diff --git a/Engine/Plugins/Developer/PerforceSourceControl/Source/PerforceSourceControl/Private/PerforceSourceControlSettings.h b/Engine/Plugins/Developer/PerforceSourceControl/Source/PerforceSourceControl/Private/PerforceSourceControlSettings.h
--- a/Engine/Plugins/Developer/PerforceSourceControl/Source/PerforceSourceControl/Private/PerforceSourceControlSettings.h
+++ b/Engine/Plugins/Developer/PerforceSourceControl/Source/PerforceSourceControl/Private/PerforceSourceControlSettings.h
@@ -5,9 +5,14 @@
 #include "CoreMinimal.h"
 #include "PerforceConnectionInfo.h"
 
+class FPerforceSourceControlProvider;
+
 class FPerforceSourceControlSettings
 {
 public:
+	FPerforceSourceControlSettings(const FPerforceSourceControlProvider& InSCCProvider, const FStringView& OwnerName);
+	~FPerforceSourceControlSettings() = default;
+
 	/** Get the Perforce port */
 	const FString& GetPort() const;
 
@@ -60,4 +65,16 @@ private:
 
 	/** The credentials we use to access the server */
 	FPerforceConnectionInfo ConnectionInfo;
+	
+	/** The section of the ini file we load our settings from */
+	FString SettingsSection;
+
+	/** Internal accessor to the source control provider associated with the object */
+	const FPerforceSourceControlProvider& GetSCCProvider() const
+	{
+		return SCCProvider;
+	}
+
+	/** The source control provider that this object is associated with */
+	const FPerforceSourceControlProvider& SCCProvider;
 };
diff --git a/Engine/Plugins/Developer/PerforceSourceControl/Source/PerforceSourceControl/Private/SPerforceSourceControlSettings.cpp b/Engine/Plugins/Developer/PerforceSourceControl/Source/PerforceSourceControl/Private/SPerforceSourceControlSettings.cpp
--- a/Engine/Plugins/Developer/PerforceSourceControl/Source/PerforceSourceControl/Private/SPerforceSourceControlSettings.cpp
+++ b/Engine/Plugins/Developer/PerforceSourceControl/Source/PerforceSourceControl/Private/SPerforceSourceControlSettings.cpp
@@ -1,17 +1,19 @@
 // Copyright Epic Games, Inc. All Rights Reserved.
 
 #include "SPerforceSourceControlSettings.h"
+
+#include "EditorStyleSet.h"
+#include "ISourceControlModule.h"
+#include "PerforceSourceControlInternalOperations.h"
 #include "PerforceSourceControlPrivate.h"
-#include "Widgets/Views/STableRow.h"
-#include "Modules/ModuleManager.h"
+#include "PerforceSourceControlProvider.h"
+#include "Widgets/Images/SImage.h"
+#include "Widgets/Images/SThrobber.h"
 #include "Widgets/Input/SButton.h"
-#include "Widgets/Views/SListView.h"
 #include "Widgets/Input/SCheckBox.h"
 #include "Widgets/Input/SComboBox.h"
-#include "EditorStyleSet.h"
-#include "ISourceControlModule.h"
-#include "PerforceSourceControlModule.h"
-#include "Widgets/Images/SThrobber.h"
+#include "Widgets/Views/SListView.h"
+#include "Widgets/Views/STableRow.h"
 
 TWeakPtr<SEditableTextBox> SPerforceSourceControlSettings::PasswordTextBox;
 
@@ -24,15 +26,16 @@ FAutoConsoleVariableRef CVarAllowP4NonTicketBasedLogins(
 	TEXT("Whether or not to allow logging in with a password directly from the perforce dialog. This is off by default because it is not a secure option. Perforce often your password as plain text in their enviroment variables")
 );
 
-void SPerforceSourceControlSettings::Construct(const FArguments& InArgs)
+void SPerforceSourceControlSettings::Construct(const FArguments& InArgs, FPerforceSourceControlProvider* InSCCProvider)
 {
-	FPerforceSourceControlModule& PerforceSourceControl = FModuleManager::LoadModuleChecked<FPerforceSourceControlModule>( "PerforceSourceControl" );
+	checkf(InSCCProvider != nullptr, TEXT("SPerforceSourceControlSettings Requires a pointer to a valid FPerforceSourceControlProvider to function"));
+	SCCProvider = InSCCProvider;
 
 	bAreAdvancedSettingsExpanded = false;
 
 	// check our settings & query if we don't already have any
-	FString PortName = PerforceSourceControl.AccessSettings().GetPort();
-	FString UserName = PerforceSourceControl.AccessSettings().GetUserName();
+	FString PortName = GetSCCProvider().AccessSettings().GetPort();
+	FString UserName = GetSCCProvider().AccessSettings().GetUserName();
 
 	if (PortName.IsEmpty() && UserName.IsEmpty())
 	{
@@ -43,10 +46,10 @@ void SPerforceSourceControlSettings::Construct(const FArguments& InArgs)
 		PortName = ANSI_TO_TCHAR(TestP4.GetPort().Text());
 		UserName = ANSI_TO_TCHAR(TestP4.GetUser().Text());
 		TestP4.Final(&P4Error);
-
-		PerforceSourceControl.AccessSettings().SetPort(PortName);
-		PerforceSourceControl.AccessSettings().SetUserName(UserName);
-		PerforceSourceControl.SaveSettings();
+		
+		GetSCCProvider().AccessSettings().SetPort(PortName);
+		GetSCCProvider().AccessSettings().SetUserName(UserName);
+		GetSCCProvider().AccessSettings().SaveSettings();
 	}
 
 	ChildSlot
@@ -232,73 +235,62 @@ void SPerforceSourceControlSettings::Tick( const FGeometry& AllottedGeometry, co
 
 ECheckBoxState SPerforceSourceControlSettings::IsP4ConfigChecked() const
 {
-	FPerforceSourceControlModule& PerforceSourceControl = FModuleManager::LoadModuleChecked<FPerforceSourceControlModule>( "PerforceSourceControl" );
-	return PerforceSourceControl.AccessSettings().GetUseP4Config() ? ECheckBoxState::Checked : ECheckBoxState::Unchecked;
+	return GetSCCProvider().AccessSettings().GetUseP4Config() ? ECheckBoxState::Checked : ECheckBoxState::Unchecked;
 }
 
 void SPerforceSourceControlSettings::OnP4ConfigCheckStatusChanged(ECheckBoxState NewState) const
 {
-	FPerforceSourceControlModule& PerforceSourceControl = FModuleManager::LoadModuleChecked<FPerforceSourceControlModule>( "PerforceSourceControl" );
-	PerforceSourceControl.AccessSettings().SetUseP4Config(NewState == ECheckBoxState::Checked);
-	PerforceSourceControl.SaveSettings();
+	GetSCCProvider().AccessSettings().SetUseP4Config(NewState == ECheckBoxState::Checked);
+	GetSCCProvider().AccessSettings().SaveSettings();
 }
 
 bool SPerforceSourceControlSettings::IsP4ConfigDisabled() const
 {
-	FPerforceSourceControlModule& PerforceSourceControl = FModuleManager::LoadModuleChecked<FPerforceSourceControlModule>( "PerforceSourceControl" );
-	return !PerforceSourceControl.AccessSettings().GetUseP4Config();
+	return !GetSCCProvider().AccessSettings().GetUseP4Config();
 }
 
 FText SPerforceSourceControlSettings::GetPortText() const
 {
-	FPerforceSourceControlModule& PerforceSourceControl = FModuleManager::LoadModuleChecked<FPerforceSourceControlModule>( "PerforceSourceControl" );
-	return FText::FromString(PerforceSourceControl.AccessSettings().GetPort());
+	return FText::FromString(GetSCCProvider().AccessSettings().GetPort());
 }
 
 void SPerforceSourceControlSettings::OnPortTextCommitted(const FText& InText, ETextCommit::Type InCommitType) const
 {
-	FPerforceSourceControlModule& PerforceSourceControl = FModuleManager::LoadModuleChecked<FPerforceSourceControlModule>( "PerforceSourceControl" );
-	PerforceSourceControl.AccessSettings().SetPort(InText.ToString());
-	PerforceSourceControl.SaveSettings();
+	GetSCCProvider().AccessSettings().SetPort(InText.ToString());
+	GetSCCProvider().AccessSettings().SaveSettings();
 }
 
 FText SPerforceSourceControlSettings::GetUserNameText() const
 {
-	FPerforceSourceControlModule& PerforceSourceControl = FModuleManager::LoadModuleChecked<FPerforceSourceControlModule>( "PerforceSourceControl" );
-	return FText::FromString(PerforceSourceControl.AccessSettings().GetUserName());
+	return FText::FromString(GetSCCProvider().AccessSettings().GetUserName());
 }
 
 void SPerforceSourceControlSettings::OnUserNameTextCommitted(const FText& InText, ETextCommit::Type InCommitType) const
 {
-	FPerforceSourceControlModule& PerforceSourceControl = FModuleManager::LoadModuleChecked<FPerforceSourceControlModule>( "PerforceSourceControl" );
-	PerforceSourceControl.AccessSettings().SetUserName(InText.ToString());
-	PerforceSourceControl.SaveSettings();
+	GetSCCProvider().AccessSettings().SetUserName(InText.ToString());
+	GetSCCProvider().AccessSettings().SaveSettings();
 }
 
 FText SPerforceSourceControlSettings::GetWorkspaceText() const
 {
-	FPerforceSourceControlModule& PerforceSourceControl = FModuleManager::LoadModuleChecked<FPerforceSourceControlModule>( "PerforceSourceControl" );
-	return FText::FromString(PerforceSourceControl.AccessSettings().GetWorkspace());
+	return FText::FromString(GetSCCProvider().AccessSettings().GetWorkspace());
 }
 
 void SPerforceSourceControlSettings::OnWorkspaceTextCommitted(const FText& InText, ETextCommit::Type InCommitType) const
 {
-	FPerforceSourceControlModule& PerforceSourceControl = FModuleManager::LoadModuleChecked<FPerforceSourceControlModule>( "PerforceSourceControl" );
-	PerforceSourceControl.AccessSettings().SetWorkspace(InText.ToString());
-	PerforceSourceControl.SaveSettings();
+	GetSCCProvider().AccessSettings().SetWorkspace(InText.ToString());
+	GetSCCProvider().AccessSettings().SaveSettings();
 }
 
 FText SPerforceSourceControlSettings::GetHostText() const
 {
-	FPerforceSourceControlModule& PerforceSourceControl = FModuleManager::LoadModuleChecked<FPerforceSourceControlModule>( "PerforceSourceControl" );
-	return FText::FromString(PerforceSourceControl.AccessSettings().GetHostOverride());
+	return FText::FromString(GetSCCProvider().AccessSettings().GetHostOverride());
 }
 
 void SPerforceSourceControlSettings::OnHostTextCommitted(const FText& InText, ETextCommit::Type InCommitType) const
 {
-	FPerforceSourceControlModule& PerforceSourceControl = FModuleManager::LoadModuleChecked<FPerforceSourceControlModule>( "PerforceSourceControl" );
-	PerforceSourceControl.AccessSettings().SetHostOverride(InText.ToString());
-	PerforceSourceControl.SaveSettings();
+	GetSCCProvider().AccessSettings().SetHostOverride(InText.ToString());
+	GetSCCProvider().AccessSettings().SaveSettings();
 }
 
 void SPerforceSourceControlSettings::QueryWorkspaces()
@@ -309,10 +301,8 @@ void SPerforceSourceControlSettings::QueryWorkspaces()
 		CurrentWorkspace = FString();
 
 		// fire off the workspace query
-		const ISourceControlModule& SourceControl = FModuleManager::LoadModuleChecked<ISourceControlModule>( "SourceControl" );
-		ISourceControlProvider& Provider = SourceControl.GetProvider();
 		GetWorkspacesOperation = ISourceControlOperation::Create<FGetWorkspaces>();
-		Provider.Execute(GetWorkspacesOperation.ToSharedRef(), EConcurrency::Asynchronous, FSourceControlOperationComplete::CreateSP(this, &SPerforceSourceControlSettings::OnSourceControlOperationComplete) );
+		GetSCCProvider().Execute(GetWorkspacesOperation.ToSharedRef(), EConcurrency::Asynchronous, FSourceControlOperationComplete::CreateSP(this, &SPerforceSourceControlSettings::OnSourceControlOperationComplete) );
 
 		State = ESourceControlOperationState::Querying;
 	}
@@ -433,9 +423,9 @@ TSharedRef<ITableRow> SPerforceSourceControlSettings::OnGenerateWorkspaceRow(TSh
 void SPerforceSourceControlSettings::OnWorkspaceSelected(TSharedPtr<FString> InItem, ESelectInfo::Type InSelectInfo)
 {
 	CurrentWorkspace = *InItem;
-	FPerforceSourceControlModule& PerforceSourceControl = FModuleManager::LoadModuleChecked<FPerforceSourceControlModule>( "PerforceSourceControl" );
-	PerforceSourceControl.AccessSettings().SetWorkspace(CurrentWorkspace);
-	PerforceSourceControl.SaveSettings();
+	
+	GetSCCProvider().AccessSettings().SetWorkspace(CurrentWorkspace);
+	GetSCCProvider().AccessSettings().SaveSettings();
 	WorkspaceCombo->SetIsOpen(false);
 }
 
@@ -448,8 +438,8 @@ FReply SPerforceSourceControlSettings::OnCancelWorkspacesRequest() const
 {
 	if(GetWorkspacesOperation.IsValid())
 	{
-		const ISourceControlModule& SourceControl = FModuleManager::LoadModuleChecked<ISourceControlModule>( "SourceControl" );
-		SourceControl.GetProvider().CancelOperation(GetWorkspacesOperation.ToSharedRef());
+		ISourceControlModule& SourceControl = FModuleManager::LoadModuleChecked<ISourceControlModule>( "SourceControl" );
+		GetSCCProvider().CancelOperation(GetWorkspacesOperation.ToSharedRef());
 	}
 	return FReply::Handled();
 }
@@ -477,4 +467,10 @@ FReply SPerforceSourceControlSettings::OnAdvancedSettingsClicked()
 	return FReply::Handled();
 }
 
+FPerforceSourceControlProvider& SPerforceSourceControlSettings::GetSCCProvider() const
+{
+	check(SCCProvider != nullptr);
+	return *SCCProvider;
+}
+
 #undef LOCTEXT_NAMESPACE
diff --git a/Engine/Plugins/Developer/PerforceSourceControl/Source/PerforceSourceControl/Private/SPerforceSourceControlSettings.h b/Engine/Plugins/Developer/PerforceSourceControl/Source/PerforceSourceControl/Private/SPerforceSourceControlSettings.h
--- a/Engine/Plugins/Developer/PerforceSourceControl/Source/PerforceSourceControl/Private/SPerforceSourceControlSettings.h
+++ b/Engine/Plugins/Developer/PerforceSourceControl/Source/PerforceSourceControl/Private/SPerforceSourceControlSettings.h
@@ -14,24 +14,12 @@
 #include "SourceControlOperationBase.h"
 #include "ISourceControlProvider.h"
 
-class SButton;
-
-/**
- * Internal-only source control operation for retrieving available workspaces.
- */
-class FGetWorkspaces : public FSourceControlOperationBase
-{
-public:
-	// ISourceControlOperation interface
-	virtual FName GetName() const override
-	{
-		return "GetWorkspaces";
-	}
+class FGetWorkspaces;
+class FPerforceSourceControlProvider;
 
-public:
-	/** Results are just an array of workspaces */
-	TArray<FString> Results;
-};
+class SButton;
+class SComboButton;
+class SEditableTextBox;
 
 /** Enum for source control operation state */
 namespace ESourceControlOperationState
@@ -49,12 +37,11 @@ class SPerforceSourceControlSettings : public SCompoundWidget
 public:
 
 	SLATE_BEGIN_ARGS(SPerforceSourceControlSettings) {}
-
 	SLATE_END_ARGS()
 
 public:
 
-	void Construct(const FArguments& InArgs);
+	void Construct(const FArguments& InArgs, FPerforceSourceControlProvider* InSCCProvider);
 
 	/** Get the currently entered password */
 	static FString GetPassword();
@@ -134,7 +121,11 @@ private:
 	/** Toggle advanced settings */
 	FReply OnAdvancedSettingsClicked();
 
+	FPerforceSourceControlProvider& GetSCCProvider() const;
+
 private:
+	FPerforceSourceControlProvider* SCCProvider;
+
 	/** State of the source control operation */
 	ESourceControlOperationState::Type State;
 
@@ -145,13 +136,13 @@ private:
 	TArray<TSharedRef<FString>> Workspaces;
 
 	/** The combo button use to display the workspaces */
-	TSharedPtr<class SComboButton> WorkspaceCombo;
+	TSharedPtr<SComboButton> WorkspaceCombo;
 
 	/** The source control operation in progress */
 	TSharedPtr<FGetWorkspaces, ESPMode::ThreadSafe> GetWorkspacesOperation;
 
 	/** Pointer to the password text box widget */
-	static TWeakPtr<class SEditableTextBox> PasswordTextBox;
+	static TWeakPtr<SEditableTextBox> PasswordTextBox;
 
 	/** Expander button for advanced settings */
 	TSharedPtr<SButton> ExpanderButton;
diff --git a/Engine/Source/Developer/SourceControl/Private/SourceControlInitSettings.cpp b/Engine/Source/Developer/SourceControl/Private/SourceControlInitSettings.cpp
new file mode 100644
--- /dev/null
+++ b/Engine/Source/Developer/SourceControl/Private/SourceControlInitSettings.cpp
@@ -0,0 +1,45 @@
+// Copyright Epic Games, Inc. All Rights Reserved.
+
+#include "SourceControlInitSettings.h"
+
+#include "Containers/StringView.h"
+
+FSourceControlInitSettings::FSourceControlInitSettings(EBehavior InBehavior)
+	: Behavior(InBehavior)
+{
+
+}
+
+void FSourceControlInitSettings::AddSetting(FStringView SettingName, FStringView SettingValue)
+{
+//	const int32 Hash = GetTypeHash(SettingName);
+//	Settings.AddByHash(Hash, SettingName, SettingValue)
+
+	Settings.Add(FString(SettingName), FString(SettingValue));
+}
+
+void FSourceControlInitSettings::OverrideSetting(FStringView SettingName, FString& InOutSettingValue)
+{
+	const int32 Hash = GetTypeHash(SettingName);
+	FString* InitialValue = Settings.FindByHash(Hash, SettingName);
+
+	if (InitialValue != nullptr)
+	{
+		InOutSettingValue = *InitialValue;
+	}
+	else if (Behavior == EBehavior::OverrideAll)
+	{
+		InOutSettingValue.Empty();
+	}
+}
+
+bool FSourceControlInitSettings::HasOverrides() const
+{
+	return !Settings.IsEmpty();
+}
+
+bool FSourceControlInitSettings::IsOverridden(FStringView SettingName) const
+{
+	const int32 Hash = GetTypeHash(SettingName);
+	return Settings.FindByHash(Hash, SettingName) != nullptr;
+}
\ No newline at end of file
diff --git a/Engine/Source/Developer/SourceControl/Private/SourceControlModule.cpp b/Engine/Source/Developer/SourceControl/Private/SourceControlModule.cpp
--- a/Engine/Source/Developer/SourceControl/Private/SourceControlModule.cpp
+++ b/Engine/Source/Developer/SourceControl/Private/SourceControlModule.cpp
@@ -334,6 +334,21 @@ ISourceControlProvider& FSourceControlModule::GetProvider() const
 	return *CurrentSourceControlProvider;
 }
 
+TUniquePtr<ISourceControlProvider> FSourceControlModule::CreateProvider(const FName& ProviderName, const FStringView& OwnerName, const FSourceControlInitSettings& InitialSettings) const
+{
+	TArray<ISourceControlProvider*> Providers = IModularFeatures::Get().GetModularFeatureImplementations<ISourceControlProvider>(SourceControlFeatureName);
+	for (const ISourceControlProvider* DefaultProvider : Providers)
+	{
+		if (DefaultProvider->GetName() == ProviderName)
+		{
+			return DefaultProvider->Create(OwnerName, InitialSettings);
+		}
+	}
+
+	// Provider was not found
+	return TUniquePtr<ISourceControlProvider>();
+}
+
 FSourceControlAssetDataCache& FSourceControlModule::GetAssetDataCache()
 {
 	return AssetDataCache;
diff --git a/Engine/Source/Developer/SourceControl/Private/SourceControlModule.h b/Engine/Source/Developer/SourceControl/Private/SourceControlModule.h
--- a/Engine/Source/Developer/SourceControl/Private/SourceControlModule.h
+++ b/Engine/Source/Developer/SourceControl/Private/SourceControlModule.h
@@ -34,6 +34,7 @@ public:
 	virtual void QueueStatusUpdate(const FString& InFilename) override;
 	virtual bool IsEnabled() const override;
 	virtual ISourceControlProvider& GetProvider() const override;
+	virtual TUniquePtr<ISourceControlProvider> CreateProvider(const FName& ProviderName, const FStringView& OwnerName, const FSourceControlInitSettings& InitialSettings) const override;
 	virtual FSourceControlAssetDataCache& GetAssetDataCache() override;
 	virtual void SetProvider( const FName& InName ) override;
 	virtual void ShowLoginDialog(const FSourceControlLoginClosed& InOnSourceControlLoginClosed, ELoginWindowMode::Type InLoginWindowMode, EOnLoginWindowStartup::Type InOnLoginWindowStartup = EOnLoginWindowStartup::ResetProviderToNone) override;
diff --git a/Engine/Source/Developer/SourceControl/Public/ISourceControlModule.h b/Engine/Source/Developer/SourceControl/Public/ISourceControlModule.h
--- a/Engine/Source/Developer/SourceControl/Public/ISourceControlModule.h
+++ b/Engine/Source/Developer/SourceControl/Public/ISourceControlModule.h
@@ -9,7 +9,10 @@
 
 class ISourceControlProvider;
 class ISourceControlChangelist;
+
 class FSourceControlAssetDataCache;
+class FSourceControlInitSettings;
+
 typedef TSharedPtr<class ISourceControlChangelist, ESPMode::ThreadSafe> FSourceControlChangelistPtr;
 
 SOURCECONTROL_API DECLARE_LOG_CATEGORY_EXTERN(LogSourceControl, Log, All);
@@ -122,6 +125,25 @@ public:
 	 */
 	virtual ISourceControlProvider& GetProvider() const = 0;
 
+	/** 
+	 * Creates and returns a unique source control provider that will be owned by
+	 * the caller. This allows subsystems to own their own source control connection
+	 * in addition to the general connection used by the rest of the process. This 
+	 * could be to avoid the user changing settings that the subsystem needs, or so
+	 * that it can open a connection with an entirely different source control type
+	 * if required.
+	 * 
+	 * @param	Name			The name of the type of provider to create.
+	 * @param	OwningSystem	The name of the system that will own the provider. This is
+	 *							used to customize things, like the config section that
+	 *							that providers settings will be saved to etc.
+	 * 
+	 * @return	A pointer to the newly created provider, this can be null if creating the 
+	 *			name is invalid or if the provider type does not support unique connection
+	 *			creation.
+	 */
+	virtual TUniquePtr<ISourceControlProvider> CreateProvider(const FName& Name, const FStringView& OwningSystem, const FSourceControlInitSettings& InitialSettings) const = 0;
+
 	/**
 	 * Get the source control AssetData information cache.
 	 */
diff --git a/Engine/Source/Developer/SourceControl/Public/ISourceControlProvider.h b/Engine/Source/Developer/SourceControl/Public/ISourceControlProvider.h
--- a/Engine/Source/Developer/SourceControl/Public/ISourceControlProvider.h
+++ b/Engine/Source/Developer/SourceControl/Public/ISourceControlProvider.h
@@ -13,6 +13,7 @@
 	#error "SOURCE_CONTROL_WITH_SLATE must be defined. Did you forget a dependency on the 'SourceControl' module?"
 #endif
 
+class FSourceControlInitSettings;
 class ISourceControlLabel;
 
 /**
@@ -327,6 +328,23 @@ public:
 	 */
 	virtual TSharedRef<class SWidget> MakeSettingsWidget() const = 0;
 #endif // SOURCE_CONTROL_WITH_SLATE
+
+protected:
+
+	/* 
+	 * The ::Create method is an easy way for us to create new providers, but we really don't 
+	 * want anything except FSourceControlModule calling it. For now we achieve this by having
+	 * ::Create protected and making FSourceControlModule a friend. 
+	 * TODO Move the create to a factory system via IModularFeatures and remove the friend declaration.
+	 */
+	friend class FSourceControlModule;
+
+	/** 
+	 * Creates a new instance of the source control provider type. Derived providers
+	 * should implement this if they can more than one provider of the same type existing. 
+	 * @see ISourceControlModule::CreateProvider
+	 */
+	virtual TUniquePtr<ISourceControlProvider> Create(const FStringView& OwnerName, const FSourceControlInitSettings& InitialSettings) const { return TUniquePtr<ISourceControlProvider>(); }
 };
 
 
diff --git a/Engine/Source/Developer/SourceControl/Public/SourceControlInitSettings.h b/Engine/Source/Developer/SourceControl/Public/SourceControlInitSettings.h
new file mode 100644
--- /dev/null
+++ b/Engine/Source/Developer/SourceControl/Public/SourceControlInitSettings.h
@@ -0,0 +1,32 @@
+// Copyright Epic Games, Inc. All Rights Reserved.
+
+#pragma once
+
+#include "Containers/Map.h"
+#include "Containers/StringFwd.h"
+#include "Containers/UnrealString.h"
+
+class SOURCECONTROL_API FSourceControlInitSettings
+{
+public:
+	enum class EBehavior
+	{
+		/** All existing settings will be overridden via the contents of FSourceControlInitSettings. Settings that are not found will be reset to default states */
+		OverrideAll,
+		/** Only the settings found in FSourceControlInitSettings will be overriden. Settings not found will be left with their current values. */
+		OverrideExisting,
+	};
+
+	FSourceControlInitSettings(EBehavior Behavior);
+	~FSourceControlInitSettings() = default;
+
+	void AddSetting(FStringView SettingName, FStringView SettingValue);
+	void OverrideSetting(FStringView SettingName, FString& InOutSettingValue);
+
+	bool HasOverrides() const;
+	bool IsOverridden(FStringView SettingName) const;
+private:
+	EBehavior Behavior;
+
+	TMap<FString, FString> Settings;
+};
