Fix problems setting perforce settings from the commandline introduced in CL 19318801.

#rnx

- The behaviour for the settings for the default provider (that takes it's settings from the command line) should be OverrideExisting not  OverrideAll to maintain the original behaviour. This was a mistake due to misreading the original code. As it was, setting one setting via the command line would clear any previously set or saved settings.
- When parsing the cmdline we used the form "KeyName=" as the key, as FParse::Value does not strip the = automatically, but when we were looking up the values to override settings we were just using "KeyName" which would not match. The lambda ::ParseCmdLineSetting now adds = to any key passed in when parsing, so there is no longer a need for the caller to add it. This helps the keys match and work correctly.

diff --git a/Engine/Plugins/Developer/PerforceSourceControl/Source/PerforceSourceControl/Private/PerforceSourceControlProvider.cpp b/Engine/Plugins/Developer/PerforceSourceControl/Source/PerforceSourceControl/Private/PerforceSourceControlProvider.cpp
--- a/Engine/Plugins/Developer/PerforceSourceControl/Source/PerforceSourceControl/Private/PerforceSourceControlProvider.cpp
+++ b/Engine/Plugins/Developer/PerforceSourceControl/Source/PerforceSourceControl/Private/PerforceSourceControlProvider.cpp
@@ -27,7 +27,7 @@ static FName ProviderName("Perforce");
 
 FPerforceSourceControlProvider::FPerforceSourceControlProvider()
 	: PerforceSCCSettings(*this, FStringView())
-	, InitialSettings(FSourceControlInitSettings::EBehavior::OverrideAll)
+	, InitialSettings(FSourceControlInitSettings::EBehavior::OverrideExisting)
 	, OwnerName(TEXT("Default"))
 	, bServerAvailable(false)
 	, bLoginError(false)
@@ -36,18 +36,18 @@ FPerforceSourceControlProvider::FPerforceSourceControlProvider()
 	auto ParseCmdLineSetting = [this](const TCHAR* SettingKey) -> void
 	{
 		FString SettingValue;
-		if (FParse::Value(FCommandLine::Get(), SettingKey, SettingValue))
+		if (FParse::Value(FCommandLine::Get(), *WriteToString<64>(SettingKey, TEXT("=")), SettingValue))
 		{
 			InitialSettings.AddSetting(SettingKey, SettingValue);
 		}
 	};
 
-	ParseCmdLineSetting(TEXT("P4Port="));
-	ParseCmdLineSetting(TEXT("P4User="));
-	ParseCmdLineSetting(TEXT("P4Client="));
-	ParseCmdLineSetting(TEXT("P4Host="));
-	ParseCmdLineSetting(TEXT("P4Passwd="));
-	ParseCmdLineSetting(TEXT("P4Changelist="));
+	ParseCmdLineSetting(TEXT("P4Port"));
+	ParseCmdLineSetting(TEXT("P4User"));
+	ParseCmdLineSetting(TEXT("P4Client"));
+	ParseCmdLineSetting(TEXT("P4Host"));
+	ParseCmdLineSetting(TEXT("P4Passwd"));
+	ParseCmdLineSetting(TEXT("P4Changelist"));
 
 	AccessSettings().LoadSettings();
 }
