FDelete now works correctly (for perforce) when passing in a ISourceControlChangelist when the operation is executed.

#rnx

- When executing FCheckOut, passing in a valid ISourceControlChangelist to ISourceControlProvider::Execute will cause the checked out files to be added to the given changelist.
- FDelete now both have the same functionality, based on the FCheckOut code.
- Note that we follow the existing convention that a changelist provided on the command line always overrides a changelist passed into ::Execute
- This code exists in some form or other in most of the workers, we probably should move it to a single function somewhere to avoid potential divergence.

diff --git a/Engine/Plugins/Developer/PerforceSourceControl/Source/PerforceSourceControl/Private/PerforceSourceControlOperations.cpp b/Engine/Plugins/Developer/PerforceSourceControl/Source/PerforceSourceControl/Private/PerforceSourceControlOperations.cpp
--- a/Engine/Plugins/Developer/PerforceSourceControl/Source/PerforceSourceControl/Private/PerforceSourceControlOperations.cpp
+++ b/Engine/Plugins/Developer/PerforceSourceControl/Source/PerforceSourceControl/Private/PerforceSourceControlOperations.cpp
@@ -774,7 +774,18 @@ bool FPerforceDeleteWorker::Execute(FPerforceSourceControlCommand& InCommand)
 		FPerforceConnection& Connection = ScopedConnection.GetConnection();
 		TArray<FString> Parameters;
 
-		AppendChangelistParameter(GetSCCProvider(), Parameters);
+		if ((!AppendChangelistParameter(GetSCCProvider(), Parameters)) && InCommand.Changelist.IsInitialized())
+		{
+			FString ChangelistNumber = InCommand.Changelist.ToString();
+
+			if (!ChangelistNumber.IsEmpty())
+			{
+				Parameters.Add(TEXT("-c"));
+				Parameters.Add(ChangelistNumber);
+				Changelist = InCommand.Changelist;
+			}
+		}
+
 		Parameters.Append(InCommand.Files);
 
 		FP4RecordSet Records;
@@ -786,6 +797,9 @@ bool FPerforceDeleteWorker::Execute(FPerforceSourceControlCommand& InCommand)
 
 bool FPerforceDeleteWorker::UpdateStates() const
 {
+	// If files have been checkedout directly to a CL, modify the cached state to reflect it.
+	UpdateChangelistState(GetSCCProvider(), Changelist, OutResults);
+
 	return UpdateCachedStates(GetSCCProvider(), OutResults);
 }
 
diff --git a/Engine/Plugins/Developer/PerforceSourceControl/Source/PerforceSourceControl/Private/PerforceSourceControlOperations.h b/Engine/Plugins/Developer/PerforceSourceControl/Source/PerforceSourceControl/Private/PerforceSourceControlOperations.h
--- a/Engine/Plugins/Developer/PerforceSourceControl/Source/PerforceSourceControl/Private/PerforceSourceControlOperations.h
+++ b/Engine/Plugins/Developer/PerforceSourceControl/Source/PerforceSourceControl/Private/PerforceSourceControlOperations.h
@@ -120,6 +120,7 @@ public:
 public:
 	/** Map of filenames to perforce state */
 	TMap<FString, EPerforceState::Type> OutResults;
+	FPerforceSourceControlChangelist Changelist;
 };
 
 class FPerforceRevertWorker final : public IPerforceSourceControlWorker
