Fixed FUpdatePendingChangelistsStatus operation that did nothing if a specific set of changelist was requested.

diff --git a/Engine/Plugins/Developer/PerforceSourceControl/Source/PerforceSourceControl/Private/PerforceSourceControlOperations.cpp b/Engine/Plugins/Developer/PerforceSourceControl/Source/PerforceSourceControl/Private/PerforceSourceControlOperations.cpp
--- a/Engine/Plugins/Developer/PerforceSourceControl/Source/PerforceSourceControl/Private/PerforceSourceControlOperations.cpp
+++ b/Engine/Plugins/Developer/PerforceSourceControl/Source/PerforceSourceControl/Private/PerforceSourceControlOperations.cpp
@@ -1999,7 +1999,11 @@ bool FPerforceGetPendingChangelistsWorker::Execute(FPerforceSourceControlCommand
 		{
 			// First, insert the default changelist which always exists
 			FPerforceSourceControlChangelistState& State = OutChangelistsStates.Emplace_GetRef(FPerforceSourceControlChangelist::DefaultChangelist);
+		}
 
+		{
+			// Query all pending change list (the default is not listed by this request) for this user/workspace. (When a specific set of changelists is requested,
+			// we could use -e to restrict the returned list, but it returns CL greater or equal to the CL specified. In general, we likely not gain much compare to list all pending ones.)
 			TArray<FString> Parameters;
 			Parameters.Add(TEXT("-l"));									// -l			Complete description
 			Parameters.Add(TEXT("-spending"));							// -s pending	Only pending changelists
@@ -2014,9 +2018,23 @@ bool FPerforceGetPendingChangelistsWorker::Execute(FPerforceSourceControlCommand
 
 			ParseChangelistsResults(Records, OutChangelistsStates);
 
-			bCleanupCache = InCommand.bCommandSuccessful;
+			// Remove the changelist that were not requested by the user.
+			if (!Operation->ShouldUpdateAllChangelists())
+			{
+				const TArray<FSourceControlChangelistRef>& RequestedChangelists = Operation->GetChangelistsToUpdate();
+				OutChangelistsStates.RemoveAll([&RequestedChangelists](const FPerforceSourceControlChangelistState& ChangelistState)
+				{
+					FPerforceSourceControlChangelistRef RemoveChangelistCandidate = StaticCastSharedRef<FPerforceSourceControlChangelist>(ChangelistState.GetChangelist());
+					return !RequestedChangelists.ContainsByPredicate([&RemoveChangelistCandidate](const FSourceControlChangelistRef& Requested)
+					{
+						return StaticCastSharedRef<FPerforceSourceControlChangelist>(Requested)->ToInt() == RemoveChangelistCandidate->ToInt();
+					});
+				});
+			}
 		}
 
+		bCleanupCache = InCommand.bCommandSuccessful;
+
 		// Test whether we should continue processing SCC commands
 		auto ShouldContinueProcessing = [&InCommand]() { return InCommand.bCommandSuccessful && !InCommand.IsCanceled(); };
 
