GetChangelistDescription SourceControlOperation to get information about the changelist for blueprint review tool

diff --git a/Engine/Plugins/Developer/PerforceSourceControl/Source/PerforceSourceControl/Private/PerforceSourceControlOperations.cpp b/Engine/Plugins/Developer/PerforceSourceControl/Source/PerforceSourceControl/Private/PerforceSourceControlOperations.cpp
--- a/Engine/Plugins/Developer/PerforceSourceControl/Source/PerforceSourceControl/Private/PerforceSourceControlOperations.cpp
+++ b/Engine/Plugins/Developer/PerforceSourceControl/Source/PerforceSourceControl/Private/PerforceSourceControlOperations.cpp
@@ -51,6 +51,7 @@ void IPerforceSourceControlWorker::RegisterWorkers()
 	WorkersMap.Add("Resolve", FGetPerforceSourceControlWorker::CreateStatic(&InstantiateWorker<FPerforceResolveWorker>));
 	WorkersMap.Add("ChangeStatus", FGetPerforceSourceControlWorker::CreateStatic(&InstantiateWorker<FPerforceChangeStatusWorker>));
 	WorkersMap.Add("UpdateChangelistsStatus", FGetPerforceSourceControlWorker::CreateStatic(&InstantiateWorker<FPerforceGetPendingChangelistsWorker>));
+	WorkersMap.Add("GetChangelistDetails", FPerforceGetChangelistDetailsWorker::CreateStatic(&InstantiateWorker<FPerforceGetChangelistDetailsWorker>));
 	WorkersMap.Add("NewChangelist", FGetPerforceSourceControlWorker::CreateStatic(&InstantiateWorker<FPerforceNewChangelistWorker>));
 	WorkersMap.Add("DeleteChangelist", FGetPerforceSourceControlWorker::CreateStatic(&InstantiateWorker<FPerforceDeleteChangelistWorker>));
 	WorkersMap.Add("EditChangelist", FGetPerforceSourceControlWorker::CreateStatic(&InstantiateWorker<FPerforceEditChangelistWorker>));
@@ -2228,7 +2229,7 @@ FName FPerforceCopyWorker::GetName() const
 	return "Copy";
 }
 
-bool FPerforceCopyWorker::Execute(class FPerforceSourceControlCommand& InCommand)
+bool FPerforceCopyWorker::Execute(FPerforceSourceControlCommand& InCommand)
 {
 	FScopedPerforceConnection ScopedConnection(InCommand);
 	if (!InCommand.IsCanceled() && ScopedConnection.IsValid())
@@ -2276,7 +2277,7 @@ FName FPerforceResolveWorker::GetName() const
 	return "Resolve";
 }
 
-bool FPerforceResolveWorker::Execute(class FPerforceSourceControlCommand& InCommand)
+bool FPerforceResolveWorker::Execute(FPerforceSourceControlCommand& InCommand)
 {
 	FScopedPerforceConnection ScopedConnection(InCommand);
 	if (!InCommand.IsCanceled() && ScopedConnection.IsValid())
@@ -2317,7 +2318,7 @@ FName FPerforceChangeStatusWorker::GetName() const
 	return "ChangeStatus";
 }
 
-bool FPerforceChangeStatusWorker::Execute(class FPerforceSourceControlCommand& InCommand)
+bool FPerforceChangeStatusWorker::Execute(FPerforceSourceControlCommand& InCommand)
 {
 	FScopedPerforceConnection ScopedConnection(InCommand);
 	if (!InCommand.IsCanceled() && ScopedConnection.IsValid())
@@ -2372,7 +2373,7 @@ FName FPerforceNewChangelistWorker::GetName() const
 	return "NewChangelist";
 }
 
-bool FPerforceNewChangelistWorker::Execute(class FPerforceSourceControlCommand& InCommand)
+bool FPerforceNewChangelistWorker::Execute(FPerforceSourceControlCommand& InCommand)
 {
 	FScopedPerforceConnection ScopedConnection(InCommand);
 
@@ -2458,7 +2459,7 @@ FName FPerforceDeleteChangelistWorker::GetName() const
 	return "DeleteChangelist";
 }
 
-bool FPerforceDeleteChangelistWorker::Execute(class FPerforceSourceControlCommand& InCommand)
+bool FPerforceDeleteChangelistWorker::Execute(FPerforceSourceControlCommand& InCommand)
 {
 	FScopedPerforceConnection ScopedConnection(InCommand);
 
@@ -2510,7 +2511,7 @@ FName FPerforceEditChangelistWorker::GetName() const
 	return "EditChangelist";
 }
 
-bool FPerforceEditChangelistWorker::Execute(class FPerforceSourceControlCommand& InCommand)
+bool FPerforceEditChangelistWorker::Execute(FPerforceSourceControlCommand& InCommand)
 {
 	FScopedPerforceConnection ScopedConnection(InCommand);
 	if (!InCommand.IsCanceled() && ScopedConnection.IsValid())
@@ -2558,7 +2559,7 @@ FName FPerforceRevertUnchangedWorker::GetName() const
 	return "RevertUnchanged";
 }
 
-bool FPerforceRevertUnchangedWorker::Execute(class FPerforceSourceControlCommand& InCommand)
+bool FPerforceRevertUnchangedWorker::Execute(FPerforceSourceControlCommand& InCommand)
 {
 	FScopedPerforceConnection ScopedConnection(InCommand);
 	if (!InCommand.IsCanceled() && ScopedConnection.IsValid())
@@ -2640,7 +2641,7 @@ FName FPerforceShelveWorker::GetName() const
 	return "Shelve";
 }
 
-bool FPerforceShelveWorker::Execute(class FPerforceSourceControlCommand& InCommand)
+bool FPerforceShelveWorker::Execute(FPerforceSourceControlCommand& InCommand)
 {
 	FScopedPerforceConnection ScopedConnection(InCommand);
 	if (!InCommand.IsCanceled() && ScopedConnection.IsValid())
@@ -2770,7 +2771,7 @@ FName FPerforceDeleteShelveWorker::GetName() const
 	return "DeleteShelved";
 }
 
-bool FPerforceDeleteShelveWorker::Execute(class FPerforceSourceControlCommand& InCommand)
+bool FPerforceDeleteShelveWorker::Execute(FPerforceSourceControlCommand& InCommand)
 {
 	FScopedPerforceConnection ScopedConnection(InCommand);
 	if (!InCommand.IsCanceled() && ScopedConnection.IsValid())
@@ -2834,7 +2835,7 @@ FName FPerforceUnshelveWorker::GetName() const
 	return "Unshelve";
 }
 
-bool FPerforceUnshelveWorker::Execute(class FPerforceSourceControlCommand& InCommand)
+bool FPerforceUnshelveWorker::Execute(FPerforceSourceControlCommand& InCommand)
 {
 	FScopedPerforceConnection ScopedConnection(InCommand);
 	if (!InCommand.IsCanceled() && ScopedConnection.IsValid())
@@ -2995,7 +2996,7 @@ FName FPerforceCreateWorkspaceWorker::GetName() const
 	return "CreateWorkspace";
 }
 
-bool FPerforceCreateWorkspaceWorker::Execute(class FPerforceSourceControlCommand& InCommand)
+bool FPerforceCreateWorkspaceWorker::Execute(FPerforceSourceControlCommand& InCommand)
 {
 	FScopedPerforceConnection ScopedConnection(InCommand);
 
@@ -3068,7 +3069,7 @@ FName FPerforceDeleteWorkspaceWorker::GetName() const
 	return "DeleteWorkspace";
 }
 
-bool FPerforceDeleteWorkspaceWorker::Execute(class FPerforceSourceControlCommand& InCommand)
+bool FPerforceDeleteWorkspaceWorker::Execute(FPerforceSourceControlCommand& InCommand)
 {
 	FScopedPerforceConnection ScopedConnection(InCommand);
 
@@ -3104,4 +3105,55 @@ bool FPerforceDeleteWorkspaceWorker::UpdateStates() const
 	return false;
 }
 
-#undef LOCTEXT_NAMESPACE
+FName FPerforceGetChangelistDetailsWorker::GetName() const
+{
+	return "GetChangelistDetails";
+}
+
+bool FPerforceGetChangelistDetailsWorker::Execute(FPerforceSourceControlCommand& InCommand)
+{
+	FScopedPerforceConnection ScopedConnection(InCommand);
+
+	if (!InCommand.IsCanceled() && ScopedConnection.IsValid())
+	{
+		FPerforceConnection& Connection = ScopedConnection.GetConnection();
+		TSharedRef<FGetChangelistDetails, ESPMode::ThreadSafe> Operation = StaticCastSharedRef<FGetChangelistDetails>(InCommand.Operation);
+
+		TArray<FString> Parameters;
+		Parameters.Add(TEXT("-s"));
+		Parameters.Add(TEXT("-S"));
+		Parameters.Add(Operation->GetChangelistNumber());
+
+		FP4RecordSet Records;
+		InCommand.bCommandSuccessful = Connection.RunCommand(TEXT("describe"), Parameters, Records, InCommand.ResultInfo.ErrorMessages, FOnIsCancelled::CreateRaw(&InCommand, &FPerforceSourceControlCommand::IsCanceled), InCommand.bConnectionDropped);
+
+		if (InCommand.bCommandSuccessful)
+		{
+			TArray<TMap<FString, FString>> ChangelistRecord;
+			ChangelistRecord.Reserve(Records.Num());
+
+			for (FP4Record& P4Record : Records)
+			{
+				TMap<FString, FString> Record;
+
+				for (TPair<FString, FString>& P4RecordPair : P4Record)
+				{
+					Record.Add(MoveTemp(P4RecordPair));
+				}
+
+				ChangelistRecord.Add(MoveTemp(Record));
+			}
+
+			Operation->SetChangelistDetails(MoveTemp(ChangelistRecord));
+		}
+	}
+
+	return InCommand.bCommandSuccessful;
+}
+
+bool FPerforceGetChangelistDetailsWorker::UpdateStates() const
+{
+	return false;
+}
+
+#undef LOCTEXT_NAMESPACE
\ No newline at end of file
diff --git a/Engine/Plugins/Developer/PerforceSourceControl/Source/PerforceSourceControl/Private/PerforceSourceControlOperations.h b/Engine/Plugins/Developer/PerforceSourceControl/Source/PerforceSourceControl/Private/PerforceSourceControlOperations.h
--- a/Engine/Plugins/Developer/PerforceSourceControl/Source/PerforceSourceControl/Private/PerforceSourceControlOperations.h
+++ b/Engine/Plugins/Developer/PerforceSourceControl/Source/PerforceSourceControl/Private/PerforceSourceControlOperations.h
@@ -8,6 +8,7 @@
 #include "PerforceSourceControlChangelistState.h"
 
 class FPerforceSourceControlRevision;
+class FPerforceSourceControlCommand;
 typedef TMap<FString, TArray< TSharedRef<FPerforceSourceControlRevision, ESPMode::ThreadSafe> > > FPerforceFileHistoryMap;
 
 class FPerforceConnectWorker final : public IPerforceSourceControlWorker
@@ -20,7 +21,7 @@ public:
 
 	// IPerforceSourceControlWorker interface
 	virtual FName GetName() const override;
-	virtual bool Execute(class FPerforceSourceControlCommand& InCommand) override;
+	virtual bool Execute(FPerforceSourceControlCommand& InCommand) override;
 	virtual bool UpdateStates() const override;
 };
 
@@ -34,7 +35,7 @@ public:
 
 	// IPerforceSourceControlWorker interface
 	virtual FName GetName() const override;
-	virtual bool Execute(class FPerforceSourceControlCommand& InCommand) override;
+	virtual bool Execute(FPerforceSourceControlCommand& InCommand) override;
 	virtual bool UpdateStates() const override;
 
 public:
@@ -53,7 +54,7 @@ public:
 
 	// IPerforceSourceControlWorker interface
 	virtual FName GetName() const override;
-	virtual bool Execute(class FPerforceSourceControlCommand& InCommand) override;
+	virtual bool Execute(FPerforceSourceControlCommand& InCommand) override;
 	virtual bool UpdateStates() const override;
 
 public:
@@ -77,7 +78,7 @@ public:
 
 	// IPerforceSourceControlWorker interface
 	virtual FName GetName() const override;
-	virtual bool Execute(class FPerforceSourceControlCommand& InCommand) override;
+	virtual bool Execute(FPerforceSourceControlCommand& InCommand) override;
 	virtual bool UpdateStates() const override;
 
 public:
@@ -95,7 +96,7 @@ public:
 
 	// IPerforceSourceControlWorker interface
 	virtual FName GetName() const override;
-	virtual bool Execute(class FPerforceSourceControlCommand& InCommand) override;
+	virtual bool Execute(FPerforceSourceControlCommand& InCommand) override;
 	virtual bool UpdateStates() const override;
 
 public:
@@ -114,7 +115,7 @@ public:
 
 	// IPerforceSourceControlWorker interface
 	virtual FName GetName() const override;
-	virtual bool Execute(class FPerforceSourceControlCommand& InCommand) override;
+	virtual bool Execute(FPerforceSourceControlCommand& InCommand) override;
 	virtual bool UpdateStates() const override;
 
 public:
@@ -133,7 +134,7 @@ public:
 
 	// IPerforceSourceControlWorker interface
 	virtual FName GetName() const override;
-	virtual bool Execute(class FPerforceSourceControlCommand& InCommand) override;
+	virtual bool Execute(FPerforceSourceControlCommand& InCommand) override;
 	virtual bool UpdateStates() const override;
 
 public:
@@ -154,7 +155,7 @@ public:
 
 	// IPerforceSourceControlWorker interface
 	virtual FName GetName() const override;
-	virtual bool Execute(class FPerforceSourceControlCommand& InCommand) override;
+	virtual bool Execute(FPerforceSourceControlCommand& InCommand) override;
 	virtual bool UpdateStates() const override;
 
 public:
@@ -172,7 +173,7 @@ public:
 
 	// IPerforceSourceControlWorker interface
 	virtual FName GetName() const override;
-	virtual bool Execute(class FPerforceSourceControlCommand& InCommand) override;
+	virtual bool Execute(FPerforceSourceControlCommand& InCommand) override;
 	virtual bool UpdateStates() const override;
 
 public:
@@ -202,7 +203,7 @@ public:
 
 	// IPerforceSourceControlWorker interface
 	virtual FName GetName() const override;
-	virtual bool Execute(class FPerforceSourceControlCommand& InCommand) override;
+	virtual bool Execute(FPerforceSourceControlCommand& InCommand) override;
 	virtual bool UpdateStates() const override;
 };
 
@@ -216,7 +217,7 @@ public:
 
 	// IPerforceSourceControlWorker interface
 	virtual FName GetName() const override;
-	virtual bool Execute(class FPerforceSourceControlCommand& InCommand) override;
+	virtual bool Execute(FPerforceSourceControlCommand& InCommand) override;
 	virtual bool UpdateStates() const override;
 
 public:
@@ -241,7 +242,7 @@ public:
 
 	// IPerforceSourceControlWorker interface
 	virtual FName GetName() const override;
-	virtual bool Execute(class FPerforceSourceControlCommand& InCommand) override;
+	virtual bool Execute(FPerforceSourceControlCommand& InCommand) override;
 	virtual bool UpdateStates() const override;
 
 public:
@@ -259,7 +260,7 @@ public:
 	virtual ~FPerforceResolveWorker() = default;
 	// IPerforceSourceControlWorker interface
 	virtual FName GetName() const override;
-	virtual bool Execute(class FPerforceSourceControlCommand& InCommand) override;
+	virtual bool Execute(FPerforceSourceControlCommand& InCommand) override;
 	virtual bool UpdateStates() const override;
 private:
 	TArray< FString > UpdatedFiles;
@@ -275,7 +276,7 @@ public:
 
 	// IPerforceSourceControlWorker interface
 	virtual FName GetName() const override;
-	virtual bool Execute(class FPerforceSourceControlCommand& InCommand) override;
+	virtual bool Execute(FPerforceSourceControlCommand& InCommand) override;
 	virtual bool UpdateStates() const override;
 };
 
@@ -287,7 +288,7 @@ public:
 
 	// IPerforceSourceControlWorker interface
 	virtual FName GetName() const override;
-	virtual bool Execute(class FPerforceSourceControlCommand& InCommand) override;
+	virtual bool Execute(FPerforceSourceControlCommand& InCommand) override;
 	virtual bool UpdateStates() const override;
 
 public:
@@ -309,7 +310,7 @@ public:
 
 	// IPerforceSourceControlWorker interface
 	virtual FName GetName() const override;
-	virtual bool Execute(class FPerforceSourceControlCommand& InCommand) override;
+	virtual bool Execute(FPerforceSourceControlCommand& InCommand) override;
 	virtual bool UpdateStates() const override;
 
 public:
@@ -326,7 +327,7 @@ public:
 
 	// IPerforceSourceControlWorker interface
 	virtual FName GetName() const override;
-	virtual bool Execute(class FPerforceSourceControlCommand& InCommand) override;
+	virtual bool Execute(FPerforceSourceControlCommand& InCommand) override;
 	virtual bool UpdateStates() const override;
 
 public:
@@ -344,7 +345,7 @@ public:
 
 	// IPerforceSourceControlWorker interface
 	virtual FName GetName() const override;
-	virtual bool Execute(class FPerforceSourceControlCommand& InCommand) override;
+	virtual bool Execute(FPerforceSourceControlCommand& InCommand) override;
 	virtual bool UpdateStates() const override;
 protected:
 	/** Map of filenames to perforce state */
@@ -364,7 +365,7 @@ public:
 
 	// IPerforceSourceControlWorker interface
 	virtual FName GetName() const override;
-	virtual bool Execute(class FPerforceSourceControlCommand& InCommand) override;
+	virtual bool Execute(FPerforceSourceControlCommand& InCommand) override;
 	virtual bool UpdateStates() const override;
 
 protected:
@@ -385,7 +386,7 @@ public:
 
 	// IPerforceSourceControlWorker interface
 	virtual FName GetName() const override;
-	virtual bool Execute(class FPerforceSourceControlCommand& InCommand) override;
+	virtual bool Execute(FPerforceSourceControlCommand& InCommand) override;
 	virtual bool UpdateStates() const override;
 
 protected:
@@ -416,7 +417,7 @@ public:
 
 	// IPerforceSourceControlWorker interface
 	virtual FName GetName() const override;
-	virtual bool Execute(class FPerforceSourceControlCommand& InCommand) override;
+	virtual bool Execute(FPerforceSourceControlCommand& InCommand) override;
 	virtual bool UpdateStates() const override;
 
 protected:
@@ -437,7 +438,7 @@ public:
 
 	// IPerforceSourceControlWorker interface
 	virtual FName GetName() const override;
-	virtual bool Execute(class FPerforceSourceControlCommand& InCommand) override;
+	virtual bool Execute(FPerforceSourceControlCommand& InCommand) override;
 	virtual bool UpdateStates() const override;
 
 protected:
@@ -458,7 +459,7 @@ public:
 
 	// IPerforceSourceControlWorker interface
 	virtual FName GetName() const override;
-	virtual bool Execute(class FPerforceSourceControlCommand& InCommand) override;
+	virtual bool Execute(FPerforceSourceControlCommand& InCommand) override;
 	virtual bool UpdateStates() const override;
 };
 
@@ -472,7 +473,7 @@ public:
 
 	// IPerforceSourceControlWorker interface
 	virtual FName GetName() const override;
-	virtual bool Execute(class FPerforceSourceControlCommand& InCommand) override;
+	virtual bool Execute(FPerforceSourceControlCommand& InCommand) override;
 	virtual bool UpdateStates() const override;
 
 private:
@@ -489,6 +490,20 @@ public:
 
 	// IPerforceSourceControlWorker interface
 	virtual FName GetName() const override;
-	virtual bool Execute(class FPerforceSourceControlCommand& InCommand) override;
+	virtual bool Execute(FPerforceSourceControlCommand& InCommand) override;
 	virtual bool UpdateStates() const override;
 };
+
+class FPerforceGetChangelistDetailsWorker final : public IPerforceSourceControlWorker
+{
+public:
+	FPerforceGetChangelistDetailsWorker(FPerforceSourceControlProvider& InSourceControlProvider)
+		: IPerforceSourceControlWorker(InSourceControlProvider)
+	{}
+	virtual ~FPerforceGetChangelistDetailsWorker() = default;
+
+	// IPerforceSourceControlWorker interface
+	virtual FName GetName() const override;
+	virtual bool Execute(FPerforceSourceControlCommand& InCommand) override;
+	virtual bool UpdateStates() const override;
+};
\ No newline at end of file
diff --git a/Engine/Source/Developer/SourceControl/Public/SourceControlOperations.h b/Engine/Source/Developer/SourceControl/Public/SourceControlOperations.h
--- a/Engine/Source/Developer/SourceControl/Public/SourceControlOperations.h
+++ b/Engine/Source/Developer/SourceControl/Public/SourceControlOperations.h
@@ -542,6 +542,43 @@ public:
 	}
 };
 
+/**
+ * This operations query the source control to extract all the details available for a given changelist. The operations returns a collection of key/value corresponding to the details available. 
+ * The list of key/value is specific to the source control implementation.
+ */
+class FGetChangelistDetails : public FSourceControlOperationBase
+{
+public:
+	// ISourceControlOperation interface
+	virtual FName GetName() const override
+	{
+		return "GetChangelistDetails";
+	}
+
+	virtual FText GetInProgressString() const override
+	{
+		return LOCTEXT("SourceControl_GetChangelistDetails", "Retrieving changelist details from Source Control...");
+	}
+
+	const FString& GetChangelistNumber() { return ChangelistNumber; }
+
+	const TArray<TMap<FString, FString>>& GetChangelistDetails() { return OutChangelistDetails; }
+	
+	void SetChangelistNumber(const FString& InChangelistNumber)
+	{
+		ChangelistNumber = InChangelistNumber;
+	}
+
+	void SetChangelistDetails(TArray<TMap<FString, FString>>&& InChangelistDetails)
+	{
+		OutChangelistDetails = MoveTemp(InChangelistDetails);
+	}
+
+private:
+	FString ChangelistNumber;
+	TArray<TMap<FString, FString>> OutChangelistDetails;
+};
+
 /**
  * Operation used to update the source control status of changelist(s)
  */
