Fix for Issue #157997: Compile errors in PerforceSourceControlOperations.cpp

diff --git a/Engine/Plugins/Developer/PerforceSourceControl/Source/PerforceSourceControl/Private/PerforceSourceControlOperations.cpp b/Engine/Plugins/Developer/PerforceSourceControl/Source/PerforceSourceControl/Private/PerforceSourceControlOperations.cpp
--- a/Engine/Plugins/Developer/PerforceSourceControl/Source/PerforceSourceControl/Private/PerforceSourceControlOperations.cpp
+++ b/Engine/Plugins/Developer/PerforceSourceControl/Source/PerforceSourceControl/Private/PerforceSourceControlOperations.cpp
@@ -51,7 +51,7 @@ void IPerforceSourceControlWorker::RegisterWorkers()
 	WorkersMap.Add("Resolve", FGetPerforceSourceControlWorker::CreateStatic(&InstantiateWorker<FPerforceResolveWorker>));
 	WorkersMap.Add("ChangeStatus", FGetPerforceSourceControlWorker::CreateStatic(&InstantiateWorker<FPerforceChangeStatusWorker>));
 	WorkersMap.Add("UpdateChangelistsStatus", FGetPerforceSourceControlWorker::CreateStatic(&InstantiateWorker<FPerforceGetPendingChangelistsWorker>));
-	WorkersMap.Add("GetChangelistDetails", FPerforceGetChangelistDetailsWorker::CreateStatic(&InstantiateWorker<FPerforceGetChangelistDetailsWorker>));
+	WorkersMap.Add("GetChangelistDetails", FGetPerforceSourceControlWorker::CreateStatic(&InstantiateWorker<FPerforceGetChangelistDetailsWorker>));
 	WorkersMap.Add("NewChangelist", FGetPerforceSourceControlWorker::CreateStatic(&InstantiateWorker<FPerforceNewChangelistWorker>));
 	WorkersMap.Add("DeleteChangelist", FGetPerforceSourceControlWorker::CreateStatic(&InstantiateWorker<FPerforceDeleteChangelistWorker>));
 	WorkersMap.Add("EditChangelist", FGetPerforceSourceControlWorker::CreateStatic(&InstantiateWorker<FPerforceEditChangelistWorker>));
