Fix FPerforceSourceControlProvider::TryToDownloadFileFromBackgroundThread not appending info/error messages to the provided operation

#rnx

- This way of invoking an operation was missing a call to FPerforceSourceControlCommand::ReturnResults which will pass the info/error messages from the command to the operation that the caller passed into the provider in the first place.
- This will allow the caller to gain access to any messages generated during the operation.

diff --git a/Engine/Plugins/Developer/PerforceSourceControl/Source/PerforceSourceControl/Private/PerforceSourceControlProvider.cpp b/Engine/Plugins/Developer/PerforceSourceControl/Source/PerforceSourceControl/Private/PerforceSourceControlProvider.cpp
--- a/Engine/Plugins/Developer/PerforceSourceControl/Source/PerforceSourceControl/Private/PerforceSourceControlProvider.cpp
+++ b/Engine/Plugins/Developer/PerforceSourceControl/Source/PerforceSourceControl/Private/PerforceSourceControlProvider.cpp
@@ -685,6 +685,11 @@ bool FPerforceSourceControlProvider::TryToDownloadFileFromBackgroundThread(const
 	check(Command.Worker->UpdateStates() == false);
 
 	OutputCommandMessages(Command);
+
+	if (!Command.bCancelledWhileTryingToConnect)
+	{
+		Command.ReturnResults();
+	}
 	
 	return Command.bCommandSuccessful;
 }
