Source control operation to get file from depot. Users can now pull a shelfed changelist files.

diff --git a/Engine/Plugins/Developer/PerforceSourceControl/Source/PerforceSourceControl/Private/PerforceSourceControlOperations.cpp b/Engine/Plugins/Developer/PerforceSourceControl/Source/PerforceSourceControl/Private/PerforceSourceControlOperations.cpp
--- a/Engine/Plugins/Developer/PerforceSourceControl/Source/PerforceSourceControl/Private/PerforceSourceControlOperations.cpp
+++ b/Engine/Plugins/Developer/PerforceSourceControl/Source/PerforceSourceControl/Private/PerforceSourceControlOperations.cpp
@@ -64,6 +64,7 @@ void IPerforceSourceControlWorker::RegisterWorkers()
 	WorkersMap.Add("CreateWorkspace", FGetPerforceSourceControlWorker::CreateStatic(&InstantiateWorker<FPerforceCreateWorkspaceWorker>));
 	WorkersMap.Add("DeleteWorkspace", FGetPerforceSourceControlWorker::CreateStatic(&InstantiateWorker<FPerforceDeleteWorkspaceWorker>));
 	WorkersMap.Add("GetFileList", FGetPerforceSourceControlWorker::CreateStatic(&InstantiateWorker<FPerforceGetFileListWorker>));
+	WorkersMap.Add("GetFile", FGetPerforceSourceControlWorker::CreateStatic(&InstantiateWorker<FPerforceGetFileListWorker>));
 }
 
 TSharedPtr<class IPerforceSourceControlWorker, ESPMode::ThreadSafe> IPerforceSourceControlWorker::CreateWorker(const FName& OperationName, FPerforceSourceControlProvider& SCCProvider)
@@ -3156,4 +3157,42 @@ bool FPerforceGetChangelistDetailsWorker::UpdateStates() const
 	return false;
 }
 
+FName FPerforceGetFileWorker::GetName() const
+{
+	return "GetFile";
+}
+
+bool FPerforceGetFileWorker::Execute(FPerforceSourceControlCommand& InCommand)
+{
+	FScopedPerforceConnection ScopedConnection(InCommand);
+
+	if (!InCommand.IsCanceled() && ScopedConnection.IsValid())
+	{
+		FPerforceConnection& Connection = ScopedConnection.GetConnection();
+		TSharedRef<FGetFile, ESPMode::ThreadSafe> Operation = StaticCastSharedRef<FGetFile>(InCommand.Operation);
+
+		TSharedRef<FPerforceSourceControlRevision, ESPMode::ThreadSafe> Revision = MakeShared<FPerforceSourceControlRevision>();
+		Revision->FileName = Operation->GetDepotFilePath();
+		Revision->ChangelistNumber = FCString::Atoi(*Operation->GetChangelistNumber());
+		Revision->RevisionNumber = FCString::Atoi(*Operation->GetRevisionNumber());
+		Revision->bIsShelve = Operation->IsShelve();
+	
+		FString OutFilename;
+
+		InCommand.bCommandSuccessful = Revision->Get(OutFilename, InCommand.Concurrency);
+
+		if (InCommand.bCommandSuccessful)
+		{
+			Operation->SetOutPackageFilename(OutFilename);
+		}
+	}
+
+	return InCommand.bCommandSuccessful;
+}
+
+bool FPerforceGetFileWorker::UpdateStates() const
+{
+	return false;
+}
+
 #undef LOCTEXT_NAMESPACE
\ No newline at end of file
diff --git a/Engine/Plugins/Developer/PerforceSourceControl/Source/PerforceSourceControl/Private/PerforceSourceControlOperations.h b/Engine/Plugins/Developer/PerforceSourceControl/Source/PerforceSourceControl/Private/PerforceSourceControlOperations.h
--- a/Engine/Plugins/Developer/PerforceSourceControl/Source/PerforceSourceControl/Private/PerforceSourceControlOperations.h
+++ b/Engine/Plugins/Developer/PerforceSourceControl/Source/PerforceSourceControl/Private/PerforceSourceControlOperations.h
@@ -502,6 +502,20 @@ public:
 	{}
 	virtual ~FPerforceGetChangelistDetailsWorker() = default;
 
+	// IPerforceSourceControlWorker interface
+	virtual FName GetName() const override;
+	virtual bool Execute(FPerforceSourceControlCommand& InCommand) override;
+	virtual bool UpdateStates() const override;
+};
+
+class FPerforceGetFileWorker final : public IPerforceSourceControlWorker
+{
+public:
+	FPerforceGetFileWorker(FPerforceSourceControlProvider& InSourceControlProvider)
+		: IPerforceSourceControlWorker(InSourceControlProvider)
+	{}
+	virtual ~FPerforceGetFileWorker() = default;
+
 	// IPerforceSourceControlWorker interface
 	virtual FName GetName() const override;
 	virtual bool Execute(FPerforceSourceControlCommand& InCommand) override;
diff --git a/Engine/Source/Developer/SourceControl/Public/SourceControlOperations.h b/Engine/Source/Developer/SourceControl/Public/SourceControlOperations.h
--- a/Engine/Source/Developer/SourceControl/Public/SourceControlOperations.h
+++ b/Engine/Source/Developer/SourceControl/Public/SourceControlOperations.h
@@ -1064,4 +1064,54 @@ protected:
 	FString WorkspaceName;
 };
 
+/**
+ * This operation uses p4v print command to get file from specified depot by shelved changelist or file revision returns a package filename
+ */
+class FGetFile : public FSourceControlOperationBase
+{
+public:
+
+	FGetFile(const FString& InChangelistNumber, const FString& InRevisionNumber, const FString& InDepotFilePath, bool bInIsShelve = false)
+		: ChangelistNumber(InChangelistNumber)
+		, RevisionNumber(InRevisionNumber)
+		, DepotFilePath(InDepotFilePath)
+		, bIsShelve(bInIsShelve)
+	{}
+
+	// ISourceControlOperation interface
+	virtual FName GetName() const override
+	{
+		return "GetFile";
+	}
+
+	virtual FText GetInProgressString() const override
+	{
+		return LOCTEXT("SourceControl_GetFile", "Retrieving file from source control...");
+	}
+
+	const FString& GetChangelistNumber() const { return ChangelistNumber; }
+
+	const FString& GetRevisionNumber() const { return RevisionNumber; }
+
+	const FString& GetDepotFilePath() const { return DepotFilePath; }
+
+	bool IsShelve() const { return bIsShelve; }
+
+	const FString& GetOutPackageFilename() const { return OutPackageFilename; }
+
+	void SetOutPackageFilename(const FString& InOutPackageFilename)
+	{
+		OutPackageFilename = InOutPackageFilename;
+	}
+
+private:
+
+	FString ChangelistNumber;
+	FString RevisionNumber;
+	FString DepotFilePath;
+	bool bIsShelve;
+
+	FString OutPackageFilename;
+};
+
 #undef LOCTEXT_NAMESPACE
