Fix a copy paste error. Source control operation to get file from depot. Users can now pull a shelfed changelist files.

diff --git a/Engine/Plugins/Developer/PerforceSourceControl/Source/PerforceSourceControl/Private/PerforceSourceControlOperations.cpp b/Engine/Plugins/Developer/PerforceSourceControl/Source/PerforceSourceControl/Private/PerforceSourceControlOperations.cpp
--- a/Engine/Plugins/Developer/PerforceSourceControl/Source/PerforceSourceControl/Private/PerforceSourceControlOperations.cpp
+++ b/Engine/Plugins/Developer/PerforceSourceControl/Source/PerforceSourceControl/Private/PerforceSourceControlOperations.cpp
@@ -64,7 +64,7 @@ void IPerforceSourceControlWorker::RegisterWorkers()
 	WorkersMap.Add("CreateWorkspace", FGetPerforceSourceControlWorker::CreateStatic(&InstantiateWorker<FPerforceCreateWorkspaceWorker>));
 	WorkersMap.Add("DeleteWorkspace", FGetPerforceSourceControlWorker::CreateStatic(&InstantiateWorker<FPerforceDeleteWorkspaceWorker>));
 	WorkersMap.Add("GetFileList", FGetPerforceSourceControlWorker::CreateStatic(&InstantiateWorker<FPerforceGetFileListWorker>));
-	WorkersMap.Add("GetFile", FGetPerforceSourceControlWorker::CreateStatic(&InstantiateWorker<FPerforceGetFileListWorker>));
+	WorkersMap.Add("GetFile", FGetPerforceSourceControlWorker::CreateStatic(&InstantiateWorker<FPerforceGetFileWorker>));
 }
 
 TSharedPtr<class IPerforceSourceControlWorker, ESPMode::ThreadSafe> IPerforceSourceControlWorker::CreateWorker(const FName& OperationName, FPerforceSourceControlProvider& SCCProvider)
