Fix compile error due to new default constructor requirement in later versions of the engine.

diff --git a/Engine/Plugins/Developer/PerforceSourceControl/Source/PerforceSourceControl/Private/PerforceSourceControlOperations.cpp b/Engine/Plugins/Developer/PerforceSourceControl/Source/PerforceSourceControl/Private/PerforceSourceControlOperations.cpp
--- a/Engine/Plugins/Developer/PerforceSourceControl/Source/PerforceSourceControl/Private/PerforceSourceControlOperations.cpp
+++ b/Engine/Plugins/Developer/PerforceSourceControl/Source/PerforceSourceControl/Private/PerforceSourceControlOperations.cpp
@@ -3171,7 +3171,7 @@ bool FPerforceGetFileWorker::Execute(FPerforceSourceControlCommand& InCommand)
 		FPerforceConnection& Connection = ScopedConnection.GetConnection();
 		TSharedRef<FGetFile, ESPMode::ThreadSafe> Operation = StaticCastSharedRef<FGetFile>(InCommand.Operation);
 
-		TSharedRef<FPerforceSourceControlRevision, ESPMode::ThreadSafe> Revision = MakeShared<FPerforceSourceControlRevision>();
+		TSharedRef<FPerforceSourceControlRevision, ESPMode::ThreadSafe> Revision = MakeShareable(new FPerforceSourceControlRevision(GetSCCProvider()));
 		Revision->FileName = Operation->GetDepotFilePath();
 		Revision->ChangelistNumber = FCString::Atoi(*Operation->GetChangelistNumber());
 		Revision->RevisionNumber = FCString::Atoi(*Operation->GetRevisionNumber());
