Combined the boolean parameters we pass to FPerforceConnection into a single bit field and add a new option to apply the global-opt '-q' to a command to supress non tagged output.

#rnx

- Some commands do not return tagged output even when we set that as our preference and will print the output to stdout as we are not currently collecting the standard output messages.
-- We could collect this output and return as info to the caller but that would involve adding more parameters which I am loath to do at the moment. At some point I plan to replace all of the FPerforceConnection::Run input and output parameters with single data structures and unify the different connection types at which point adding support for returning all info output would make sense.
-  To supress this output we would normally pass -q as a global-opt for the command, which in the case of the cpp API requires us to call ::SetQuiet on the client users but for that we need to be able to pass in the request to FPerforceConnection::Run.
-- Note that calling ::SetQuiet is pretty much the same as overriding ClientUser::OuputInfo
- At this point we had two bool parameters in ::Run already, so when adding a third we might as well convert this to a bitfield of flags.
- It seems that the former parameter bInAllowRetry was never hooked to anything since the last large source conde API refactor near a decade a go, so that could just be removed.
- I felt that the former parameter 'bInStandardDebugOutput' was not clearly named, in the end if it was true we would log the full p4 command (and timing of the command in VeryVerbose mode) and in almost all cases this was set to true. So I flipped the flag so it is not set to opt out of logging a command.
- I removed one overload of FPerforceConnection::Run by setting the flags to have a default value. This should help discourage anyone from adding more parameters to the method.

diff --git a/Engine/Plugins/Developer/PerforceSourceControl/Source/PerforceSourceControl/Private/PerforceConnection.cpp b/Engine/Plugins/Developer/PerforceSourceControl/Source/PerforceSourceControl/Private/PerforceConnection.cpp
--- a/Engine/Plugins/Developer/PerforceSourceControl/Source/PerforceSourceControl/Private/PerforceConnection.cpp
+++ b/Engine/Plugins/Developer/PerforceSourceControl/Source/PerforceSourceControl/Private/PerforceConnection.cpp
@@ -21,9 +21,11 @@ namespace
 {
 	enum class EP4ClientUserFlags
 	{
-		None = 0,
-		UnicodeServer = 1 << 0,
-		CollectData = 1 << 1
+		None			= 0,
+		/** The server uses unicode */
+		UnicodeServer	= 1 << 0,
+		/** Binary data returned by commands should be collected in the DataBuffer member */
+		CollectData		= 1 << 1,
 	};
 
 	ENUM_CLASS_FLAGS(EP4ClientUserFlags);
@@ -801,13 +803,17 @@ void FPerforceConnection::Disconnect()
 	}
 }
 
-bool FPerforceConnection::RunCommand(const FString& InCommand, const TArray<FString>& InParameters, FP4RecordSet& OutRecordSet, TOptional<FSharedBuffer>& OutData, TArray<FText>& OutErrorMessage, FOnIsCancelled InIsCancelled, bool& OutConnectionDropped, const bool bInStandardDebugOutput, const bool bInAllowRetry)
+bool FPerforceConnection::RunCommand(	const FString& InCommand, const TArray<FString>& InParameters, FP4RecordSet& OutRecordSet, 
+										TOptional<FSharedBuffer>& OutData, TArray<FText>& OutErrorMessage, 
+										FOnIsCancelled InIsCancelled, bool& OutConnectionDropped, ERunCommandFlags RunFlags)
 {
 	if (!bEstablishedConnection)
 	{
 		return false;
 	}
 
+	const bool bLogCommandDetails = !EnumHasAllFlags(RunFlags, ERunCommandFlags::DisableCommandLogging);
+
 	FString FullCommand = InCommand;
 
 	// Prepare arguments
@@ -827,14 +833,14 @@ bool FPerforceConnection::RunCommand(const FString& InCommand, const TArray<FStr
 			FMemory::Memcpy(ArgV[Index], TCHAR_TO_ANSI(*InParameters[Index]), InParameters[Index].Len() + 1);
 		}
 		
-		if (bInStandardDebugOutput)
+		if (bLogCommandDetails)
 		{
 			FullCommand += TEXT(" ");
 			FullCommand += InParameters[Index];
 		}
 	}
 
-	if (bInStandardDebugOutput)
+	if (bLogCommandDetails)
 	{
 		UE_LOG( LogSourceControl, Log, TEXT("Attempting 'p4 %s'"), *FullCommand );
 	}
@@ -848,11 +854,16 @@ bool FPerforceConnection::RunCommand(const FString& InCommand, const TArray<FStr
 
 	OutRecordSet.Reset();
 
-	EP4ClientUserFlags Flags = EP4ClientUserFlags::None;
-	Flags |= bIsUnicode ? EP4ClientUserFlags::UnicodeServer : EP4ClientUserFlags::None;
-	Flags |= OutData ? EP4ClientUserFlags::CollectData : EP4ClientUserFlags::None;
+	EP4ClientUserFlags ClientUserFlags = EP4ClientUserFlags::None;
+	ClientUserFlags |= bIsUnicode ? EP4ClientUserFlags::UnicodeServer : EP4ClientUserFlags::None;
+	ClientUserFlags |= OutData ? EP4ClientUserFlags::CollectData : EP4ClientUserFlags::None;
 	
-	FP4ClientUser User(OutRecordSet, Flags, OutErrorMessage);
+	FP4ClientUser User(OutRecordSet, ClientUserFlags, OutErrorMessage);
+	if (EnumHasAllFlags(RunFlags, ERunCommandFlags::Quiet))
+	{
+		User.SetQuiet();
+	}
+
 	P4Client.Run(FROM_TCHAR(*InCommand, bIsUnicode), &User);
 	if ( P4Client.Dropped() )
 	{
@@ -879,7 +890,7 @@ bool FPerforceConnection::RunCommand(const FString& InCommand, const TArray<FStr
 		SCCProvider.SetLastErrors(OutErrorMessage);
 	}
 
-	if (bInStandardDebugOutput)
+	if (bLogCommandDetails)
 	{
 		UE_LOG( LogSourceControl, VeryVerbose, TEXT("P4 execution time: %0.4f seconds. Command: %s"), FPlatformTime::Seconds() - SCCStartTime, *FullCommand );
 	}
@@ -916,14 +927,13 @@ int32 FPerforceConnection::EditPendingChangelist(const FText& NewDescription, in
 		TArray<FString> Params;
 		TOptional<FSharedBuffer> CommandData;
 		bool bConnectionDropped = false;
-		const bool bStandardDebugOutput = false;
-		const bool bAllowRetry = true;
+		const ERunCommandFlags RunFlags = ERunCommandFlags::DisableCommandLogging;
 
 		Params.Add(TEXT("-o"));
 		// TODO : make this work also for default changelist, but should really be a Create
 		Params.Add(FString::Printf(TEXT("%d"), ChangelistNumber));
 
-		if (!RunCommand(TEXT("change"), Params, PreviousRecords, CommandData, OutErrorMessages, InIsCancelled, bConnectionDropped, bStandardDebugOutput, bAllowRetry))
+		if (!RunCommand(TEXT("change"), Params, PreviousRecords, CommandData, OutErrorMessages, InIsCancelled, bConnectionDropped, RunFlags))
 		{
 			return 0;
 		}
@@ -1021,12 +1031,11 @@ void FPerforceConnection::EstablishConnection(const FPerforceConnectionInfo& InC
 		FP4RecordSet Records;
 		TOptional<FSharedBuffer> CommandData;
 		bool bConnectionDropped = false;
-		const bool bStandardDebugOutput = false;
-		const bool bAllowRetry = true;
+		const ERunCommandFlags RunFlags = ERunCommandFlags::DisableCommandLogging;
 
 		UE_LOG(LogSourceControl, Verbose, TEXT(" ... checking unicode status" ));
 
-		if (RunCommand(TEXT("info"), Params, Records, CommandData, ErrorMessages, FOnIsCancelled(), bConnectionDropped, bStandardDebugOutput, bAllowRetry))
+		if (RunCommand(TEXT("info"), Params, Records, CommandData, ErrorMessages, FOnIsCancelled(), bConnectionDropped, RunFlags))
 		{
 			// Get character encoding
 			bIsUnicode = Records[0].Find(TEXT("unicode")) != nullptr;
@@ -1059,7 +1068,7 @@ void FPerforceConnection::EstablishConnection(const FPerforceConnectionInfo& InC
 			// Gather the client root
 			UE_LOG(LogSourceControl, Verbose, TEXT(" ... getting info" ));
 			bConnectionDropped = false;
-			if (RunCommand(TEXT("info"), Params, Records, CommandData, ErrorMessages, FOnIsCancelled(), bConnectionDropped, bStandardDebugOutput, bAllowRetry))
+			if (RunCommand(TEXT("info"), Params, Records, CommandData, ErrorMessages, FOnIsCancelled(), bConnectionDropped, RunFlags))
 			{
 				UE_LOG(LogSourceControl, Verbose, TEXT(" ... getting clientroot" ));
 				ClientRoot = Records[0](TEXT("clientRoot"));
diff --git a/Engine/Plugins/Developer/PerforceSourceControl/Source/PerforceSourceControl/Private/PerforceConnection.h b/Engine/Plugins/Developer/PerforceSourceControl/Source/PerforceSourceControl/Private/PerforceConnection.h
--- a/Engine/Plugins/Developer/PerforceSourceControl/Source/PerforceSourceControl/Private/PerforceConnection.h
+++ b/Engine/Plugins/Developer/PerforceSourceControl/Source/PerforceSourceControl/Private/PerforceConnection.h
@@ -35,6 +35,18 @@ enum class EConnectionOptions : uint8
 };
 ENUM_CLASS_FLAGS(EConnectionOptions);
 
+/** Optional flags that can be passed in when a command is run */
+enum class ERunCommandFlags : uint32
+{
+	/** No special options selected */
+	Default					= 0,
+	/** Prevents us writing the full perforce command (once evaluated) to the log file */
+	DisableCommandLogging		= 1 << 0,
+	/** The same as passing -q to a perforce command as part of the g-opts. Prevents info output from being printed to stdout */
+	Quiet					= 1 << 1
+};
+ENUM_CLASS_FLAGS(ERunCommandFlags);
+
 class FPerforceConnection
 {
 public:
@@ -81,24 +93,14 @@ public:
 	 */
 	bool RunCommand(const FString& InCommand, const TArray<FString>& InParameters, FP4RecordSet& OutRecordSet, TArray<FText>&  OutErrorMessage, FOnIsCancelled InIsCancelled, bool& OutConnectionDropped)
 	{
-		const bool bStandardDebugOutput=true;
-		const bool bAllowRetry=true;
 		TOptional<FSharedBuffer> UnsetDataBuffer;
-
-		return RunCommand(InCommand, InParameters, OutRecordSet, UnsetDataBuffer, OutErrorMessage, InIsCancelled, OutConnectionDropped, bStandardDebugOutput, bAllowRetry);
-	}
-
-	bool RunCommand(const FString& InCommand, const TArray<FString>& InParameters, FP4RecordSet& OutRecordSet, TOptional<FSharedBuffer>& OutData, TArray<FText>& OutErrorMessage, FOnIsCancelled InIsCancelled, bool& OutConnectionDropped)
-	{
-		const bool bStandardDebugOutput = true;
-		const bool bAllowRetry = true;
-		return RunCommand(InCommand, InParameters, OutRecordSet, OutData, OutErrorMessage, InIsCancelled, OutConnectionDropped, bStandardDebugOutput, bAllowRetry);
+		return RunCommand(InCommand, InParameters, OutRecordSet, UnsetDataBuffer, OutErrorMessage, InIsCancelled, OutConnectionDropped);
 	}
 
 	/**
 	 * Runs internal perforce command, catches exceptions, returns results
 	 */
-	bool RunCommand(const FString& InCommand, const TArray<FString>& InParameters, FP4RecordSet& OutRecordSet, TOptional<FSharedBuffer>& OutData, TArray<FText>& OutErrorMessage, FOnIsCancelled InIsCancelled, bool& OutConnectionDropped, const bool bInStandardDebugOutput, const bool bInAllowRetry);
+	bool RunCommand(const FString& InCommand, const TArray<FString>& InParameters, FP4RecordSet& OutRecordSet, TOptional<FSharedBuffer>& OutData, TArray<FText>& OutErrorMessage, FOnIsCancelled InIsCancelled, bool& OutConnectionDropped, ERunCommandFlags RunFlags = ERunCommandFlags::Default);
 
 	/**
 	 * Creates a changelist with the specified description
diff --git a/Engine/Plugins/Developer/PerforceSourceControl/Source/PerforceSourceControl/Private/PerforceSourceControlOperations.cpp b/Engine/Plugins/Developer/PerforceSourceControl/Source/PerforceSourceControl/Private/PerforceSourceControlOperations.cpp
--- a/Engine/Plugins/Developer/PerforceSourceControl/Source/PerforceSourceControl/Private/PerforceSourceControlOperations.cpp
+++ b/Engine/Plugins/Developer/PerforceSourceControl/Source/PerforceSourceControl/Private/PerforceSourceControlOperations.cpp
@@ -2927,10 +2927,11 @@ bool FPerforceDownloadFileWorker::Execute(FPerforceSourceControlCommand& InComma
 				Parameters.Add(TargetFilePath);
 
 				TOptional<FSharedBuffer> FileData = FSharedBuffer();
+				const ERunCommandFlags Flags = Operation->ShouldLogToStdOutput() ? ERunCommandFlags::Default : ERunCommandFlags::DisableCommandLogging;
 
 				InCommand.bCommandSuccessful = Connection.RunCommand(	TEXT("print"), Parameters, Records, FileData, InCommand.ResultInfo.ErrorMessages, 
 																		FOnIsCancelled::CreateRaw(&InCommand, &FPerforceSourceControlCommand::IsCanceled), 
-																		InCommand.bConnectionDropped, Operation->ShouldLogToStdOutput(), true);
+																		InCommand.bConnectionDropped, Flags);
 				
 				if (InCommand.bCommandSuccessful)
 				{
@@ -2959,10 +2960,11 @@ bool FPerforceDownloadFileWorker::Execute(FPerforceSourceControlCommand& InComma
 				Parameters.Add(TargetFilePath);
 
 				TOptional<FSharedBuffer> NullBuffer;
+				const ERunCommandFlags Flags = Operation->ShouldLogToStdOutput() ? ERunCommandFlags::Default : ERunCommandFlags::DisableCommandLogging;
 
 				InCommand.bCommandSuccessful = Connection.RunCommand(	TEXT("print"), Parameters, Records, NullBuffer, InCommand.ResultInfo.ErrorMessages, 
 																		FOnIsCancelled::CreateRaw(&InCommand, &FPerforceSourceControlCommand::IsCanceled), 
-																		InCommand.bConnectionDropped, Operation->ShouldLogToStdOutput(), true);
+																		InCommand.bConnectionDropped, Flags);
 
 				if (InCommand.bCommandSuccessful)
 				{
