Running FCreateWorkspace/FDeleteWorkspace with the perforce source control provider will no longer print info directly to stdout.

#rnx

- Because the commands do not return data in tagged format it was being written directly to stdout, this would cause odd log messages stating that a workspace had been create/deleted that would be displayed outside of our normal logging format.

diff --git a/Engine/Plugins/Developer/PerforceSourceControl/Source/PerforceSourceControl/Private/PerforceConnection.cpp b/Engine/Plugins/Developer/PerforceSourceControl/Source/PerforceSourceControl/Private/PerforceConnection.cpp
--- a/Engine/Plugins/Developer/PerforceSourceControl/Source/PerforceSourceControl/Private/PerforceConnection.cpp
+++ b/Engine/Plugins/Developer/PerforceSourceControl/Source/PerforceSourceControl/Private/PerforceConnection.cpp
@@ -974,6 +974,9 @@ bool FPerforceConnection::CreateWorkspace(FStringView WorkspaceSpec, FOnIsCancel
 	P4Client.SetBreak(&KeepAlive);
 
 	FP4CommandWithStdInputClientUser User(WorkspaceSpec, Records, Flags, OutErrorMessages, P4Client);
+	User.SetQuiet();	// p4 client does not return tagged output, so any output messages will be
+						// printed to stdout. Setting this will prevent that.
+
 	P4Client.Run("client", &User);
 
 	P4Client.SetBreak(nullptr);
diff --git a/Engine/Plugins/Developer/PerforceSourceControl/Source/PerforceSourceControl/Private/PerforceSourceControlOperations.cpp b/Engine/Plugins/Developer/PerforceSourceControl/Source/PerforceSourceControl/Private/PerforceSourceControlOperations.cpp
--- a/Engine/Plugins/Developer/PerforceSourceControl/Source/PerforceSourceControl/Private/PerforceSourceControlOperations.cpp
+++ b/Engine/Plugins/Developer/PerforceSourceControl/Source/PerforceSourceControl/Private/PerforceSourceControlOperations.cpp
@@ -3087,13 +3087,20 @@ bool FPerforceDeleteWorkspaceWorker::Execute(FPerforceSourceControlCommand& InCo
 		Parameters.Add(Operation->GetWorkspaceName());
 
 		FP4RecordSet Records;
+		TOptional<FSharedBuffer> NullBuffer;
+		// p4 client command does not return tagged results, which means info will get printed to
+		// stdout which we do not want, so we can use the Quiet flag which is the same as passing
+		// in -q as a global-opt in the command.
+		const ERunCommandFlags RunFlags = ERunCommandFlags::Quiet;
 
 		Connection.RunCommand(	TEXT("client"), 
 								Parameters, 
 								Records, 
+								NullBuffer,
 								InCommand.ResultInfo.ErrorMessages, 
 								FOnIsCancelled::CreateRaw(&InCommand, &FPerforceSourceControlCommand::IsCanceled), 
-								InCommand.bConnectionDropped);
+								InCommand.bConnectionDropped,
+								RunFlags);
 
 		// p4 client -d doesn't return any records, so the return value of RunCommand will always
 		// be false, so check the error messages instead.
