Fix some duplicate localization keys

#rnx

diff --git a/Engine/Plugins/Developer/PerforceSourceControl/Source/PerforceSourceControl/Private/PerforceConnection.cpp b/Engine/Plugins/Developer/PerforceSourceControl/Source/PerforceSourceControl/Private/PerforceConnection.cpp
--- a/Engine/Plugins/Developer/PerforceSourceControl/Source/PerforceSourceControl/Private/PerforceConnection.cpp
+++ b/Engine/Plugins/Developer/PerforceSourceControl/Source/PerforceSourceControl/Private/PerforceConnection.cpp
@@ -538,6 +538,12 @@ bool FPerforceConnection::EnsureValidConnection(FString& InOutServerName, FStrin
 		}
 	}
 
+	// Add easy access to the localized error message if needed
+	auto GetFailedToConnectMessage = []() -> FText
+	{
+		return LOCTEXT("P4ErrorConnection_FailedToConnect", "P4ERROR: Failed to connect to source control provider.");
+	};
+
 	Error P4Error;
 	TestP4.Init(&P4Error);
 
@@ -547,7 +553,7 @@ bool FPerforceConnection::EnsureValidConnection(FString& InOutServerName, FStrin
 		//Connection FAILED
 		StrBuf ErrorMessage;
 		P4Error.Fmt(&ErrorMessage);
-		SourceControlLog.Error(LOCTEXT("P4ErrorConnection_FailedToConnect", "P4ERROR: Failed to connect to source control provider."));
+		SourceControlLog.Error(GetFailedToConnectMessage());
 		SourceControlLog.Error(FText::FromString(ANSI_TO_TCHAR(ErrorMessage.Text())));
 		FFormatNamedArguments Arguments;
 		Arguments.Add(TEXT("OwningSystem"), FText::FromString(SCCProvider.GetOwnerName()));
@@ -555,7 +561,7 @@ bool FPerforceConnection::EnsureValidConnection(FString& InOutServerName, FStrin
 		Arguments.Add( TEXT("UserName"), FText::FromString(NewUserName) );
 		Arguments.Add( TEXT("ClientSpecName"), FText::FromString(NewClientSpecName) );
 		Arguments.Add( TEXT("Ticket"), FText::FromString(InConnectionInfo.Ticket) );
-		SourceControlLog.Error(FText::Format(LOCTEXT("P4ErrorConnection_Details", "OwningSystem={OwningSystem}, Port={PortName}, User={UserName}, ClientSpec={ClientSpecName}, Ticket={Ticket}"), Arguments));
+		SourceControlLog.Error(FText::Format(LOCTEXT("P4ConnectErrorConnection_Details", "OwningSystem={OwningSystem}, Port={PortName}, User={UserName}, ClientSpec={ClientSpecName}, Ticket={Ticket}"), Arguments));
 	}
 
 	// run an info command to determine unicode status
@@ -574,7 +580,7 @@ bool FPerforceConnection::EnsureValidConnection(FString& InOutServerName, FStrin
 			Arguments.Add( TEXT("UserName"), FText::FromString(NewUserName) );
 			Arguments.Add( TEXT("ClientSpecName"), FText::FromString(NewClientSpecName) );
 			Arguments.Add( TEXT("Ticket"), FText::FromString(InConnectionInfo.Ticket) );
-			SourceControlLog.Error(FText::Format(LOCTEXT("P4ErrorConnection_Details", "OwningSystem={OwningSystem}, Port={PortName}, User={UserName}, ClientSpec={ClientSpecName}, Ticket={Ticket}"), Arguments));
+			SourceControlLog.Error(FText::Format(LOCTEXT("P4UnicodeErrorConnection_Details", "OwningSystem={OwningSystem}, Port={PortName}, User={UserName}, ClientSpec={ClientSpecName}, Ticket={Ticket}"), Arguments));
 		}
 		else
 		{
@@ -604,14 +610,14 @@ bool FPerforceConnection::EnsureValidConnection(FString& InOutServerName, FStrin
 			FString ServerName = TO_TCHAR(TestP4.GetPort().Text(), bIsUnicodeServer);
 			FString UserName = TO_TCHAR(TestP4.GetUser().Text(), bIsUnicodeServer);
 
-			SourceControlLog.Error(LOCTEXT("P4ErrorConnection_FailedToConnect", "P4ERROR: Failed to connect to source control provider."));
+			SourceControlLog.Error(GetFailedToConnectMessage());
 			SourceControlLog.Error(ErrorMessages.Num() > 0 ? ErrorMessages[0] : LOCTEXT("P4ErrorConnection_InvalidToken", "Unable to log in"));
 			FFormatNamedArguments Arguments;
 			Arguments.Add(TEXT("OwningSystem"), FText::FromString(SCCProvider.GetOwnerName()));
 			Arguments.Add(TEXT("PortName"), FText::FromString(MoveTemp(ServerName)));
 			Arguments.Add(TEXT("UserName"), FText::FromString(MoveTemp(UserName)));
 
-			SourceControlLog.Error(FText::Format(LOCTEXT("P4ErrorConnection_Details", "OwningSystem={OwningSystem}, Port={PortName}, User={UserName}"), Arguments));
+			SourceControlLog.Error(FText::Format(LOCTEXT("P4LoginErrorConnection_Details", "OwningSystem={OwningSystem}, Port={PortName}, User={UserName}"), Arguments));
 		}
 	}
 
@@ -643,7 +649,7 @@ bool FPerforceConnection::EnsureValidConnection(FString& InOutServerName, FStrin
 		
 		if (!bConnectionOK)
 		{
-			SourceControlLog.Error(LOCTEXT("P4ErrorConnection_FailedToConnect", "P4ERROR: Failed to connect to source control provider."));
+			SourceControlLog.Error(GetFailedToConnectMessage());
 			SourceControlLog.Error(ErrorMessages.Num() > 0 ? ErrorMessages[0] : LOCTEXT("P4ErrorConnection_InvalidWorkspace", "Invalid workspace"));
 			FFormatNamedArguments Arguments;
 			Arguments.Add(TEXT("OwningSystem"), FText::FromString(SCCProvider.GetOwnerName()));
@@ -651,7 +657,7 @@ bool FPerforceConnection::EnsureValidConnection(FString& InOutServerName, FStrin
 			Arguments.Add( TEXT("UserName"), FText::FromString(NewUserName) );
 			Arguments.Add( TEXT("ClientSpecName"), FText::FromString(NewClientSpecName) );
 			Arguments.Add( TEXT("Ticket"), FText::FromString(InConnectionInfo.Ticket) );
-			SourceControlLog.Error(FText::Format(LOCTEXT("P4ErrorConnection_Details", "OwningSystem={OwningSystem}, Port={PortName}, User={UserName}, ClientSpec={ClientSpecName}, Ticket={Ticket}"), Arguments));
+			SourceControlLog.Error(FText::Format(LOCTEXT("P4ClientErrorConnection_Details", "OwningSystem={OwningSystem}, Port={PortName}, User={UserName}, ClientSpec={ClientSpecName}, Ticket={Ticket}"), Arguments));
 		}
 	}
 
