Added TRACE_CPUPROFILER_EVENT_SCOPE to Perforce integration plugin to be able to capture some performance metrics.

diff --git a/Engine/Plugins/Developer/PerforceSourceControl/Source/PerforceSourceControl/Private/PerforceConnection.cpp b/Engine/Plugins/Developer/PerforceSourceControl/Source/PerforceSourceControl/Private/PerforceConnection.cpp
--- a/Engine/Plugins/Developer/PerforceSourceControl/Source/PerforceSourceControl/Private/PerforceConnection.cpp
+++ b/Engine/Plugins/Developer/PerforceSourceControl/Source/PerforceSourceControl/Private/PerforceConnection.cpp
@@ -11,6 +11,7 @@
 #include "Misc/Paths.h"
 #include "PerforceSourceControlProvider.h"
 #include "PerforceSourceControlSettings.h"
+#include "ProfilingDebugging/CpuProfilerTrace.h"
 
 #define LOCTEXT_NAMESPACE "PerforceConnection"
 
@@ -355,6 +356,8 @@ public:
 
 static bool TestLoginConnection(FPerforceSourceControlProvider& SCCProvider, ClientApi& P4Client, bool bIsUnicodeServer, TArray<FText>& OutErrorMessages)
 {
+	TRACE_CPUPROFILER_EVENT_SCOPE(FPerforce::TestLoginConnection);
+
 	FP4RecordSet Records;
 	EP4ClientUserFlags Flags = bIsUnicodeServer ? EP4ClientUserFlags::UnicodeServer : EP4ClientUserFlags::None;
 	FP4ClientUser User(Records, Flags, OutErrorMessages);
@@ -375,6 +378,8 @@ static bool TestLoginConnection(FPerforceSourceControlProvider& SCCProvider, Cli
  */
 static bool TestClientConnection(FPerforceSourceControlProvider& SCCProvider, ClientApi& P4Client, const FString& ClientSpecName, bool bIsUnicodeServer, TArray<FText>& OutErrorMessages)
 {
+	TRACE_CPUPROFILER_EVENT_SCOPE(FPerforce::TestClientConnection);
+
 	FP4RecordSet Records;
 	EP4ClientUserFlags Flags = bIsUnicodeServer ? EP4ClientUserFlags::UnicodeServer : EP4ClientUserFlags::None;
 	FP4ClientUser User(Records, Flags, OutErrorMessages);
@@ -705,6 +710,8 @@ bool FPerforceConnection::EnsureValidConnection(FString& InOutServerName, FStrin
 
 bool FPerforceConnection::GetWorkspaceList(const FPerforceConnectionInfo& InConnectionInfo, FOnIsCancelled InOnIsCanceled, TArray<FString>& OutWorkspaceList, TArray<FText>& OutErrorMessages)
 {
+	TRACE_CPUPROFILER_EVENT_SCOPE(FPerforceConnection::GetWorkspaceList)
+
 	if(bEstablishedConnection)
 	{
 		TArray<FString> Params;
@@ -813,6 +820,8 @@ bool FPerforceConnection::RunCommand(	const FString& InCommand, const TArray<FSt
 										TOptional<FSharedBuffer>& OutData, TArray<FText>& OutErrorMessage, 
 										FOnIsCancelled InIsCancelled, bool& OutConnectionDropped, ERunCommandFlags RunFlags)
 {
+	TRACE_CPUPROFILER_EVENT_SCOPE_TEXT(*FString::Printf(TEXT("FPerforceConnection::RunCommand_%s"), *InCommand));
+
 	if (!bEstablishedConnection)
 	{
 		return false;
@@ -968,6 +977,8 @@ int32 FPerforceConnection::EditPendingChangelist(const FText& NewDescription, in
 
 bool FPerforceConnection::CreateWorkspace(FStringView WorkspaceSpec, FOnIsCancelled InIsCancelled, TArray<FText>& OutErrorMessages)
 {
+	TRACE_CPUPROFILER_EVENT_SCOPE(FPerforceConnection::CreateWorkspace);
+
 	TArray<FString> Params;
 	FP4RecordSet Records;
 
@@ -992,6 +1003,8 @@ bool FPerforceConnection::CreateWorkspace(FStringView WorkspaceSpec, FOnIsCancel
 
 void FPerforceConnection::EstablishConnection(const FPerforceConnectionInfo& InConnectionInfo)
 {
+	TRACE_CPUPROFILER_EVENT_SCOPE(FPerforceConnection::EstablishConnection);
+
 	// Verify Input. ServerName and UserName are required
 	if ( InConnectionInfo.Port.IsEmpty() || InConnectionInfo.UserName.IsEmpty() )
 	{
diff --git a/Engine/Plugins/Developer/PerforceSourceControl/Source/PerforceSourceControl/Private/PerforceSourceControlOperations.cpp b/Engine/Plugins/Developer/PerforceSourceControl/Source/PerforceSourceControl/Private/PerforceSourceControlOperations.cpp
--- a/Engine/Plugins/Developer/PerforceSourceControl/Source/PerforceSourceControl/Private/PerforceSourceControlOperations.cpp
+++ b/Engine/Plugins/Developer/PerforceSourceControl/Source/PerforceSourceControl/Private/PerforceSourceControlOperations.cpp
@@ -20,6 +20,7 @@
 #include "PerforceSourceControlRevision.h"
 #include "SourceControlHelpers.h"
 #include "SourceControlOperations.h"
+#include "ProfilingDebugging/CpuProfilerTrace.h"
 
 #define LOCTEXT_NAMESPACE "PerforceSourceControl"
 
@@ -242,6 +243,8 @@ static void ParseRecordSetForState(const FP4RecordSet& InRecords, TMap<FString,
 
 static void UpdateChangelistState(FPerforceSourceControlProvider& SCCProvider, const FPerforceSourceControlChangelist& InChangelist, const TMap<FString, EPerforceState::Type>& InOperationResults)
 {
+	TRACE_CPUPROFILER_EVENT_SCOPE(FPerforce::UpdateChangelistState);
+
 	if (InChangelist.IsInitialized())
 	{
 		TSharedRef<FPerforceSourceControlChangelistState, ESPMode::ThreadSafe> ChangelistState = SCCProvider.GetStateInternal(InChangelist);
@@ -261,6 +264,8 @@ static void UpdateChangelistState(FPerforceSourceControlProvider& SCCProvider, c
 
 static bool UpdateCachedStates(FPerforceSourceControlProvider& SCCProvider, const TMap<FString, EPerforceState::Type>& InResults)
 {
+	TRACE_CPUPROFILER_EVENT_SCOPE(FPerforce::UpdateCachedStates);
+
 	for(TMap<FString, EPerforceState::Type>::TConstIterator It(InResults); It; ++It)
 	{
 		TSharedRef<FPerforceSourceControlState, ESPMode::ThreadSafe> State = SCCProvider.GetStateInternal(It.Key());
@@ -330,6 +335,8 @@ FName FPerforceConnectWorker::GetName() const
 
 bool FPerforceConnectWorker::Execute(FPerforceSourceControlCommand& InCommand)
 {
+	TRACE_CPUPROFILER_EVENT_SCOPE(FPerforceConnectWorker::Execute);
+
 	FScopedPerforceConnection ScopedConnection(InCommand);
 	if (!InCommand.IsCanceled() && ScopedConnection.IsValid())
 	{
@@ -381,6 +388,8 @@ FName FPerforceCheckOutWorker::GetName() const
 
 bool FPerforceCheckOutWorker::Execute(FPerforceSourceControlCommand& InCommand)
 {
+	TRACE_CPUPROFILER_EVENT_SCOPE(FPerforceCheckOutWorker::Execute);
+
 	FScopedPerforceConnection ScopedConnection(InCommand);
 	if (!InCommand.IsCanceled() && ScopedConnection.IsValid())
 	{
@@ -410,6 +419,8 @@ bool FPerforceCheckOutWorker::Execute(FPerforceSourceControlCommand& InCommand)
 
 bool FPerforceCheckOutWorker::UpdateStates() const
 {
+	TRACE_CPUPROFILER_EVENT_SCOPE(FPerforceCheckOutWorker::UpdateStates);
+
 	// If files have been checkedout directly to a CL, modify the cached state to reflect it.
 	UpdateChangelistState(GetSCCProvider(), InChangelist, OutResults);
 
@@ -439,6 +450,8 @@ static FText ParseSubmitResults(const FP4RecordSet& InRecords)
 
 static bool RunReopenCommand(FPerforceSourceControlCommand& InCommand, const TArray<FString>& InFiles, const FPerforceSourceControlChangelist& InChangelist, TArray<FString>* OutReopenedFiles = nullptr)
 {
+	TRACE_CPUPROFILER_EVENT_SCOPE(FPerforce::RunReopenCommand);
+
 	bool bCommandSuccessful = true;
 
 	FScopedPerforceConnection ScopedConnection(InCommand);
@@ -502,6 +515,8 @@ static bool RemoveFilesFromChangelist(const TMap<FString, EPerforceState::Type>&
 
 bool FPerforceCheckInWorker::Execute(FPerforceSourceControlCommand& InCommand)
 {
+	TRACE_CPUPROFILER_EVENT_SCOPE(FPerforceCheckInWorker::Execute);
+
 	FScopedPerforceConnection ScopedConnection(InCommand);
 	if (!InCommand.IsCanceled() && ScopedConnection.IsValid())
 	{
@@ -629,6 +644,8 @@ bool FPerforceCheckInWorker::Execute(FPerforceSourceControlCommand& InCommand)
 
 bool FPerforceCheckInWorker::UpdateStates() const
 {
+	TRACE_CPUPROFILER_EVENT_SCOPE(FPerforceCheckInWorker::UpdateStates);
+
 	bool bUpdatedStates = UpdateCachedStates(GetSCCProvider(), OutResults);
 	bool bUpdatedChangelistStates = false;
 
@@ -662,6 +679,8 @@ static void AppendMaskParameter(TArray<FString>& InOutParams)
 
 bool FPerforceGetFileListWorker::Execute(FPerforceSourceControlCommand& InCommand)
 {
+	TRACE_CPUPROFILER_EVENT_SCOPE(FPerforceGetFileListWorker::Execute);
+
 	FScopedPerforceConnection ScopedConnection(InCommand);
 	if (!InCommand.IsCanceled() && ScopedConnection.IsValid())
 	{
@@ -703,6 +722,8 @@ FName FPerforceMarkForAddWorker::GetName() const
 
 bool FPerforceMarkForAddWorker::Execute(FPerforceSourceControlCommand& InCommand)
 {
+	TRACE_CPUPROFILER_EVENT_SCOPE(FPerforceMarkForAddWorker::Execute);
+
 	// Avoid invalid p4 syntax if there's no file to process
 	if (InCommand.Files.IsEmpty())
 	{
@@ -758,6 +779,8 @@ bool FPerforceMarkForAddWorker::Execute(FPerforceSourceControlCommand& InCommand
 
 bool FPerforceMarkForAddWorker::UpdateStates() const
 {
+	TRACE_CPUPROFILER_EVENT_SCOPE(FPerforceMarkForAddWorker::UpdateStates);
+
 	UpdateChangelistState(GetSCCProvider(), InChangelist, OutResults);
 
 	return UpdateCachedStates(GetSCCProvider(), OutResults);
@@ -770,6 +793,8 @@ FName FPerforceDeleteWorker::GetName() const
 
 bool FPerforceDeleteWorker::Execute(FPerforceSourceControlCommand& InCommand)
 {
+	TRACE_CPUPROFILER_EVENT_SCOPE(FPerforceDeleteWorker::Execute);
+
 	FScopedPerforceConnection ScopedConnection(InCommand);
 	if (!InCommand.IsCanceled() && ScopedConnection.IsValid())
 	{
@@ -799,6 +824,8 @@ bool FPerforceDeleteWorker::Execute(FPerforceSourceControlCommand& InCommand)
 
 bool FPerforceDeleteWorker::UpdateStates() const
 {
+	TRACE_CPUPROFILER_EVENT_SCOPE(FPerforceDeleteWorker::UpdateStates);
+
 	// If files have been checkedout directly to a CL, modify the cached state to reflect it.
 	UpdateChangelistState(GetSCCProvider(), Changelist, OutResults);
 
@@ -812,6 +839,8 @@ FName FPerforceRevertWorker::GetName() const
 
 bool FPerforceRevertWorker::Execute(FPerforceSourceControlCommand& InCommand)
 {
+	TRACE_CPUPROFILER_EVENT_SCOPE(FPerforceRevertWorker::Execute);
+
 	FScopedPerforceConnection ScopedConnection(InCommand);
 	if (!InCommand.IsCanceled() && ScopedConnection.IsValid())
 	{
@@ -867,6 +896,8 @@ bool FPerforceRevertWorker::Execute(FPerforceSourceControlCommand& InCommand)
 
 bool FPerforceRevertWorker::UpdateStates() const
 {
+	TRACE_CPUPROFILER_EVENT_SCOPE(FPerforceRevertWorker::UpdateStates);
+
 	bool bUpdatedCachedStates = UpdateCachedStates(GetSCCProvider(), OutResults);
 	bool bUpdatedChangelists = false;
 
@@ -913,6 +944,8 @@ static void ParseSyncResults(const FP4RecordSet& InRecords, TMap<FString, EPerfo
 
 bool FPerforceSyncWorker::Execute(FPerforceSourceControlCommand& InCommand)
 {
+	TRACE_CPUPROFILER_EVENT_SCOPE(FPerforceSyncWorker::Execute);
+
 	FScopedPerforceConnection ScopedConnection(InCommand);
 	if (!InCommand.IsCanceled() && ScopedConnection.IsValid())
 	{
@@ -976,12 +1009,15 @@ bool FPerforceSyncWorker::Execute(FPerforceSourceControlCommand& InCommand)
 
 bool FPerforceSyncWorker::UpdateStates() const
 {
+	TRACE_CPUPROFILER_EVENT_SCOPE(FPerforceSyncWorker::UpdateStates);
+
 	return UpdateCachedStates(GetSCCProvider(), OutResults);
 }
 
 static void ParseBranchModificationResults(const FP4RecordSet& InRecords, const TArray<FText>& ErrorMessages, const FString& ContentRoot, TMap<FString, FBranchModification>& BranchModifications)
 {
-
+	TRACE_CPUPROFILER_EVENT_SCOPE(FPerforce::ParseBranchModificationResults);
+		
 	for (int32 Index = 0; Index < InRecords.Num(); ++Index)
 	{
 		const FP4Record& ClientRecord = InRecords[Index];
@@ -1068,7 +1104,8 @@ static void ParseBranchModificationResults(const FP4RecordSet& InRecords, const
 
 static void ParseUpdateStatusResults(const FP4RecordSet& InRecords, const TArray<FText>& ErrorMessages, TArray<FPerforceSourceControlState>& OutStates, const FString& ContentRoot, TMap<FString, FBranchModification>& BranchModifications)
 {
-	TRACE_CPUPROFILER_EVENT_SCOPE(FPerforceUpdateStatusWorker_ParseUpdateStatusResults);
+	TRACE_CPUPROFILER_EVENT_SCOPE(FPerforceUpdateStatusWorker::ParseUpdateStatusResults);
+
 	// Build up a map of any other branch states
 	for (int32 Index = 0; Index < InRecords.Num(); ++Index)
 	{
@@ -1640,6 +1677,8 @@ static void ParseHistoryResults(FPerforceSourceControlProvider& SCCProvider, con
 static bool GetFileHistory(	FPerforceSourceControlProvider& SCCProvider, FPerforceConnection& Connection, FPerforceSourceControlCommand& InCommand, 
 							const TArray<FString>& InFiles, TArray<FPerforceSourceControlState>& OutStates, FPerforceFileHistoryMap& OutHistory)
 {
+	TRACE_CPUPROFILER_EVENT_SCOPE(FPerforce::GetFileHistory);
+
 	TArray<FString> Parameters;
 	FP4RecordSet Records;
 	// disregard non-contributory integrations
@@ -1703,6 +1742,7 @@ FName FPerforceUpdateStatusWorker::GetName() const
 bool FPerforceUpdateStatusWorker::Execute(FPerforceSourceControlCommand& InCommand)
 {
 	TRACE_CPUPROFILER_EVENT_SCOPE(FPerforceUpdateStatusWorker::Execute);
+
 	FScopedPerforceConnection ScopedConnection(InCommand);
 	if (!InCommand.IsCanceled() && ScopedConnection.IsValid())
 	{
@@ -1848,6 +1888,8 @@ bool FPerforceUpdateStatusWorker::Execute(FPerforceSourceControlCommand& InComma
 
 bool FPerforceUpdateStatusWorker::UpdateStates() const
 {
+	TRACE_CPUPROFILER_EVENT_SCOPE(FPerforceUpdateStatusWorker::UpdateStates);
+
 	bool bUpdated = false;
 
 	const FDateTime Now = FDateTime::Now();
@@ -1899,6 +1941,8 @@ FName FPerforceGetWorkspacesWorker::GetName() const
 
 bool FPerforceGetWorkspacesWorker::Execute(FPerforceSourceControlCommand& InCommand)
 {
+	TRACE_CPUPROFILER_EVENT_SCOPE(FPerforceGetWorkspacesWorker::Execute);
+
 	FScopedPerforceConnection ScopedConnection(InCommand);
 	if (!InCommand.IsCanceled() && ScopedConnection.IsValid())
 	{
@@ -1926,6 +1970,8 @@ FName FPerforceGetPendingChangelistsWorker::GetName() const
 
 static bool GetOpenedFilesInChangelist(FPerforceConnection& Connection, FPerforceSourceControlCommand& InCommand, const FPerforceSourceControlChangelist& Changelist, TArray<FPerforceSourceControlState>& FilesStates)
 {
+	TRACE_CPUPROFILER_EVENT_SCOPE(FPerforce::GetOpenedFilesInChangelist);
+
 	TArray<FString> Parameters;
 	Parameters.Add(TEXT("-c"));	// -c	Changelist
 	Parameters.Add(Changelist.ToString());	// <changelist>
@@ -1960,6 +2006,8 @@ static void ParseWhereResults(FP4RecordSet& InRecords, TMap<FString, FString>& D
 
 static bool GetDepotFileToLocalFileMap(FPerforceConnection& Connection, FPerforceSourceControlCommand& InCommand, const TMap<FString, EPerforceState::Type>& InDepotFiles, TMap<FString, FString>& OutDepotToLocalMap)
 {
+	TRACE_CPUPROFILER_EVENT_SCOPE(FPerforce::GetDepotFileToLocalFileMap);
+
 	if (InDepotFiles.Num() == 0)
 	{
 		return true;
@@ -1987,6 +2035,8 @@ static bool GetDepotFileToLocalFileMap(FPerforceConnection& Connection, FPerforc
 
 bool FPerforceGetPendingChangelistsWorker::Execute(FPerforceSourceControlCommand& InCommand)
 {
+	TRACE_CPUPROFILER_EVENT_SCOPE(FPerforceGetPendingChangelistsWorker::Execute);
+
 	FScopedPerforceConnection ScopedConnection(InCommand);
 
 	if (!InCommand.IsCanceled() && ScopedConnection.IsValid())
@@ -2107,6 +2157,8 @@ static bool AddShelvedFilesToChangelist(FPerforceSourceControlProvider& SCCProvi
 										const TMap<FString, FString>& DepotToFileMap, TSharedRef<FPerforceSourceControlChangelistState, ESPMode::ThreadSafe>& ChangelistState, 
 										const FDateTime* TimeStamp = nullptr)
 {
+	TRACE_CPUPROFILER_EVENT_SCOPE(FPerforce::AddShelvedFilesToChangelist);
+
 	const FDateTime Now = (TimeStamp ? *TimeStamp : FDateTime::Now());
 
 	for (TMap<FString, EPerforceState::Type>::TConstIterator It(FilesToAdd); It; ++It)
@@ -2163,6 +2215,8 @@ static bool AddShelvedFilesToChangelist(FPerforceSourceControlProvider& SCCProvi
 
 bool FPerforceGetPendingChangelistsWorker::UpdateStates() const
 {
+	TRACE_CPUPROFILER_EVENT_SCOPE(FPerforceGetPendingChangelistsWorker::UpdateStates);
+
 	bool bUpdated = false;
 
 	const FDateTime Now = FDateTime::Now();
@@ -2232,6 +2286,8 @@ FName FPerforceCopyWorker::GetName() const
 
 bool FPerforceCopyWorker::Execute(FPerforceSourceControlCommand& InCommand)
 {
+	TRACE_CPUPROFILER_EVENT_SCOPE(FPerforceCopyWorker::Execute);
+
 	FScopedPerforceConnection ScopedConnection(InCommand);
 	if (!InCommand.IsCanceled() && ScopedConnection.IsValid())
 	{
@@ -2269,6 +2325,7 @@ bool FPerforceCopyWorker::Execute(FPerforceSourceControlCommand& InCommand)
 
 bool FPerforceCopyWorker::UpdateStates() const
 {
+	TRACE_CPUPROFILER_EVENT_SCOPE(FPerforceCopyWorker::UpdateStates);
 	return UpdateCachedStates(GetSCCProvider(), OutResults);
 }
 
@@ -2280,6 +2337,8 @@ FName FPerforceResolveWorker::GetName() const
 
 bool FPerforceResolveWorker::Execute(FPerforceSourceControlCommand& InCommand)
 {
+	TRACE_CPUPROFILER_EVENT_SCOPE(FPerforceResolveWorker::Execute);
+
 	FScopedPerforceConnection ScopedConnection(InCommand);
 	if (!InCommand.IsCanceled() && ScopedConnection.IsValid())
 	{
@@ -2304,6 +2363,8 @@ bool FPerforceResolveWorker::Execute(FPerforceSourceControlCommand& InCommand)
 
 bool FPerforceResolveWorker::UpdateStates() const
 {
+	TRACE_CPUPROFILER_EVENT_SCOPE(FPerforceResolveWorker::UpdateStates);
+
 	for( const auto& Filename : UpdatedFiles )
 	{
 		auto State = GetSCCProvider().GetStateInternal( Filename );
@@ -2321,6 +2382,8 @@ FName FPerforceChangeStatusWorker::GetName() const
 
 bool FPerforceChangeStatusWorker::Execute(FPerforceSourceControlCommand& InCommand)
 {
+	TRACE_CPUPROFILER_EVENT_SCOPE(FPerforceChangeStatusWorker::Execute);
+
 	FScopedPerforceConnection ScopedConnection(InCommand);
 	if (!InCommand.IsCanceled() && ScopedConnection.IsValid())
 	{
@@ -2376,6 +2439,8 @@ FName FPerforceNewChangelistWorker::GetName() const
 
 bool FPerforceNewChangelistWorker::Execute(FPerforceSourceControlCommand& InCommand)
 {
+	TRACE_CPUPROFILER_EVENT_SCOPE(FPerforceNewChangelistWorker::Execute);
+
 	FScopedPerforceConnection ScopedConnection(InCommand);
 
 	if (!InCommand.IsCanceled() && ScopedConnection.IsValid())
@@ -2422,6 +2487,8 @@ bool FPerforceNewChangelistWorker::Execute(FPerforceSourceControlCommand& InComm
 
 bool FPerforceNewChangelistWorker::UpdateStates() const
 {
+	TRACE_CPUPROFILER_EVENT_SCOPE(FPerforceNewChangelistWorker::UpdateStates);
+
 	if (NewChangelist.IsInitialized())
 	{
 		const FDateTime Now = FDateTime::Now();
@@ -2462,6 +2529,8 @@ FName FPerforceDeleteChangelistWorker::GetName() const
 
 bool FPerforceDeleteChangelistWorker::Execute(FPerforceSourceControlCommand& InCommand)
 {
+	TRACE_CPUPROFILER_EVENT_SCOPE(FPerforceDeleteChangelistWorker::Execute);
+
 	FScopedPerforceConnection ScopedConnection(InCommand);
 
 	// Can't delete the default changelist
@@ -2496,6 +2565,8 @@ bool FPerforceDeleteChangelistWorker::Execute(FPerforceSourceControlCommand& InC
 
 bool FPerforceDeleteChangelistWorker::UpdateStates() const
 {
+	TRACE_CPUPROFILER_EVENT_SCOPE(FPerforceDeleteChangelistWorker::UpdateStates);
+
 	if (!DeletedChangelist.IsDefault())
 	{
 		return GetSCCProvider().RemoveChangelistFromCache(DeletedChangelist);
@@ -2514,6 +2585,8 @@ FName FPerforceEditChangelistWorker::GetName() const
 
 bool FPerforceEditChangelistWorker::Execute(FPerforceSourceControlCommand& InCommand)
 {
+	TRACE_CPUPROFILER_EVENT_SCOPE(FPerforceEditChangelistWorker::Execute);
+
 	FScopedPerforceConnection ScopedConnection(InCommand);
 	if (!InCommand.IsCanceled() && ScopedConnection.IsValid())
 	{
@@ -2546,6 +2619,8 @@ bool FPerforceEditChangelistWorker::Execute(FPerforceSourceControlCommand& InCom
 
 bool FPerforceEditChangelistWorker::UpdateStates() const
 {
+	TRACE_CPUPROFILER_EVENT_SCOPE(FPerforceEditChangelistWorker::UpdateStates);
+
 	TSharedRef<FPerforceSourceControlChangelistState, ESPMode::ThreadSafe> EditedChangelistState = GetSCCProvider().GetStateInternal(EditedChangelist);
 	// TODO: update similar to NewChangelist when/if we support files in edit/new changelists.
 	EditedChangelistState->Description = EditedDescription.ToString();
@@ -2562,6 +2637,8 @@ FName FPerforceRevertUnchangedWorker::GetName() const
 
 bool FPerforceRevertUnchangedWorker::Execute(FPerforceSourceControlCommand& InCommand)
 {
+	TRACE_CPUPROFILER_EVENT_SCOPE(FPerforceRevertUnchangedWorker::Execute);
+
 	FScopedPerforceConnection ScopedConnection(InCommand);
 	if (!InCommand.IsCanceled() && ScopedConnection.IsValid())
 	{
@@ -2587,6 +2664,8 @@ bool FPerforceRevertUnchangedWorker::Execute(FPerforceSourceControlCommand& InCo
 
 bool FPerforceRevertUnchangedWorker::UpdateStates() const
 {
+	TRACE_CPUPROFILER_EVENT_SCOPE(FPerforceRevertUnchangedWorker::UpdateStates);
+
 	bool bUpdatedStates = UpdateCachedStates(GetSCCProvider(), OutResults);
 	bool bUpdatedChangelistState = ChangelistToUpdate.IsInitialized() && RemoveFilesFromChangelist(OutResults, GetSCCProvider(), ChangelistToUpdate);
 	return bUpdatedStates || bUpdatedChangelistState;
@@ -2599,6 +2678,8 @@ FName FPerforceReopenWorker::GetName() const
 	
 bool FPerforceReopenWorker::Execute(FPerforceSourceControlCommand& InCommand)
 {
+	TRACE_CPUPROFILER_EVENT_SCOPE(FPerforceReopenWorker::Execute);
+
 	FScopedPerforceConnection ScopedConnection(InCommand);
 	if (!InCommand.IsCanceled() && ScopedConnection.IsValid())
 	{
@@ -2612,6 +2693,8 @@ bool FPerforceReopenWorker::Execute(FPerforceSourceControlCommand& InCommand)
 
 bool FPerforceReopenWorker::UpdateStates() const
 {
+	TRACE_CPUPROFILER_EVENT_SCOPE(FPerforceReopenWorker::UpdateStates);
+
 	const FDateTime Now = FDateTime::Now();
 	
 	TSharedRef<FPerforceSourceControlChangelistState, ESPMode::ThreadSafe> DestinationChangelistState = GetSCCProvider().GetStateInternal(DestinationChangelist);
@@ -2644,6 +2727,8 @@ FName FPerforceShelveWorker::GetName() const
 
 bool FPerforceShelveWorker::Execute(FPerforceSourceControlCommand& InCommand)
 {
+	TRACE_CPUPROFILER_EVENT_SCOPE(FPerforceShelveWorker::Execute);
+
 	FScopedPerforceConnection ScopedConnection(InCommand);
 	if (!InCommand.IsCanceled() && ScopedConnection.IsValid())
 	{
@@ -2735,6 +2820,8 @@ bool FPerforceShelveWorker::Execute(FPerforceSourceControlCommand& InCommand)
 
 bool FPerforceShelveWorker::UpdateStates() const
 {
+	TRACE_CPUPROFILER_EVENT_SCOPE(FPerforceShelveWorker::UpdateStates);
+
 	bool bMovedFiles = false;
 
 	// If we moved files to a new changelist, then we must make sure that the files are properly moved
@@ -2774,6 +2861,8 @@ FName FPerforceDeleteShelveWorker::GetName() const
 
 bool FPerforceDeleteShelveWorker::Execute(FPerforceSourceControlCommand& InCommand)
 {
+	TRACE_CPUPROFILER_EVENT_SCOPE(FPerforceDeleteShelveWorker::Execute);
+
 	FScopedPerforceConnection ScopedConnection(InCommand);
 	if (!InCommand.IsCanceled() && ScopedConnection.IsValid())
 	{
@@ -2804,6 +2893,8 @@ bool FPerforceDeleteShelveWorker::Execute(FPerforceSourceControlCommand& InComma
 
 bool FPerforceDeleteShelveWorker::UpdateStates() const
 {
+	TRACE_CPUPROFILER_EVENT_SCOPE(FPerforceDeleteShelveWorker::UpdateStates);
+
 	if (ChangelistToUpdate.IsInitialized())
 	{
 		TSharedRef<FPerforceSourceControlChangelistState, ESPMode::ThreadSafe> ChangelistState = GetSCCProvider().GetStateInternal(ChangelistToUpdate);
@@ -2838,6 +2929,8 @@ FName FPerforceUnshelveWorker::GetName() const
 
 bool FPerforceUnshelveWorker::Execute(FPerforceSourceControlCommand& InCommand)
 {
+	TRACE_CPUPROFILER_EVENT_SCOPE(FPerforceUnshelveWorker::Execute);
+
 	FScopedPerforceConnection ScopedConnection(InCommand);
 	if (!InCommand.IsCanceled() && ScopedConnection.IsValid())
 	{
@@ -2874,6 +2967,8 @@ bool FPerforceUnshelveWorker::Execute(FPerforceSourceControlCommand& InCommand)
 
 bool FPerforceUnshelveWorker::UpdateStates() const
 {
+	TRACE_CPUPROFILER_EVENT_SCOPE(FPerforceUnshelveWorker::UpdateStates);
+
 	if (ChangelistToUpdate.IsInitialized() && ChangelistFilesStates.Num() > 0)
 	{
 		const FDateTime Now = FDateTime::Now();
@@ -2904,6 +2999,8 @@ FName FPerforceDownloadFileWorker::GetName() const
 
 bool FPerforceDownloadFileWorker::Execute(FPerforceSourceControlCommand& InCommand)
 {
+	TRACE_CPUPROFILER_EVENT_SCOPE(FPerforceDownloadFileWorker::Execute);
+
 	FScopedPerforceConnection ScopedConnection(InCommand);
 
 	if (!InCommand.IsCanceled() && ScopedConnection.IsValid())
@@ -3001,6 +3098,8 @@ FName FPerforceCreateWorkspaceWorker::GetName() const
 
 bool FPerforceCreateWorkspaceWorker::Execute(FPerforceSourceControlCommand& InCommand)
 {
+	TRACE_CPUPROFILER_EVENT_SCOPE(FPerforceCreateWorkspaceWorker::Execute);
+
 	FScopedPerforceConnection ScopedConnection(InCommand);
 
 	if (!InCommand.IsCanceled() && ScopedConnection.IsValid())
@@ -3074,6 +3173,8 @@ FName FPerforceDeleteWorkspaceWorker::GetName() const
 
 bool FPerforceDeleteWorkspaceWorker::Execute(FPerforceSourceControlCommand& InCommand)
 {
+	TRACE_CPUPROFILER_EVENT_SCOPE(FPerforceDeleteWorkspaceWorker::Execute);
+
 	FScopedPerforceConnection ScopedConnection(InCommand);
 
 	if (!InCommand.IsCanceled() && ScopedConnection.IsValid())
@@ -3122,6 +3223,8 @@ FName FPerforceGetChangelistDetailsWorker::GetName() const
 
 bool FPerforceGetChangelistDetailsWorker::Execute(FPerforceSourceControlCommand& InCommand)
 {
+	TRACE_CPUPROFILER_EVENT_SCOPE(FPerforceGetChangelistDetailsWorker::Execute);
+
 	FScopedPerforceConnection ScopedConnection(InCommand);
 
 	if (!InCommand.IsCanceled() && ScopedConnection.IsValid())
@@ -3173,6 +3276,8 @@ FName FPerforceGetFileWorker::GetName() const
 
 bool FPerforceGetFileWorker::Execute(FPerforceSourceControlCommand& InCommand)
 {
+	TRACE_CPUPROFILER_EVENT_SCOPE(FPerforceGetChangelistDetailsWorker::FPerforceGetFileWorker);
+
 	FScopedPerforceConnection ScopedConnection(InCommand);
 
 	if (!InCommand.IsCanceled() && ScopedConnection.IsValid())
@@ -3204,4 +3309,4 @@ bool FPerforceGetFileWorker::UpdateStates() const
 	return false;
 }
 
-#undef LOCTEXT_NAMESPACE
\ No newline at end of file
+#undef LOCTEXT_NAMESPACE
