Added a CPU timer scope into PerforceSourceControlProvider::Tick() to capture time spend when a state update is broadcasted.

diff --git a/Engine/Plugins/Developer/PerforceSourceControl/Source/PerforceSourceControl/Private/PerforceSourceControlProvider.cpp b/Engine/Plugins/Developer/PerforceSourceControl/Source/PerforceSourceControl/Private/PerforceSourceControlProvider.cpp
--- a/Engine/Plugins/Developer/PerforceSourceControl/Source/PerforceSourceControl/Private/PerforceSourceControlProvider.cpp
+++ b/Engine/Plugins/Developer/PerforceSourceControl/Source/PerforceSourceControl/Private/PerforceSourceControlProvider.cpp
@@ -573,6 +573,7 @@ void FPerforceSourceControlProvider::Tick()
 
 	if(bStatesUpdated)
 	{
+		TRACE_CPUPROFILER_EVENT_SCOPE(FPerforceSourceControlProvider::Tick::BroadcastStateUpdate);
 		OnSourceControlStateChanged.Broadcast();
 	}
 }
