Fix duplicate line
#rnx

diff --git a/Engine/Plugins/Developer/PerforceSourceControl/Source/PerforceSourceControl/Private/PerforceSourceControlOperations.cpp b/Engine/Plugins/Developer/PerforceSourceControl/Source/PerforceSourceControl/Private/PerforceSourceControlOperations.cpp
--- a/Engine/Plugins/Developer/PerforceSourceControl/Source/PerforceSourceControl/Private/PerforceSourceControlOperations.cpp
+++ b/Engine/Plugins/Developer/PerforceSourceControl/Source/PerforceSourceControl/Private/PerforceSourceControlOperations.cpp
@@ -2302,7 +2302,6 @@ bool FPerforceCopyWorker::Execute(FPerforceSourceControlCommand& InCommand)
 
 		AppendChangelistParameter(GetSCCProvider(), Parameters);
 
-		FP4RecordSet Records;
 		Parameters.Append(InCommand.Files);
 		Parameters.Add(DestinationPath);
 
