Fix issues when trying to run a perforce connection without a clientspec

#rnx

### Perforce Connection
- The test for bRequireWorkspace was inverted, we only want to require a workspace if the workspace is NOT optional.
- Combined the two if statements checking if we should try and auto assign a workspace to reduce indentation.
- Changed use of Len() == 0 to IsEmpty() for consistency
- Do not take the workspace name from the default environment if bRequireWorkspace is false

### PerforceSourceControlProvider
- When switching workspaces use EConnectionOptions::WorkspaceOptional when establishing the connection since we are providing the workspace name.
- Remove the check, the entire point of ::EnsureValidConnection is that it can change the workspace requested.

### PerforceSourceControloOperations
- Running FConnect command with perforce does not actually connect, but rather validates the clientspec. This needed to be changed to not give an error if the clientspec is empty.
- I'd rather allow FConnect to work for connections without a clientspec rather than avoid calling it at a higher level in case a different source control provider would require FConnect to be called at this point. Not calling it requires too much knowledge of the target source control provider to be safe.

### Future Work
- Currently the optional workspace flag is used to indicate that the workspace is being set for the connection manually and we don't want to auto evalute the workspace name. However we are really confusing two different issues here 1) having no client 2) having a specific client. In the second case we might want to allow defaulting to a valid workspace. So overall it would be better to have a flag indicating that it is okay to have no clientspec name and another flag indicating that we want to fallback to a valid name if the one we provided was not valid.

diff --git a/Engine/Plugins/Developer/PerforceSourceControl/Source/PerforceSourceControl/Private/PerforceConnection.cpp b/Engine/Plugins/Developer/PerforceSourceControl/Source/PerforceSourceControl/Private/PerforceConnection.cpp
--- a/Engine/Plugins/Developer/PerforceSourceControl/Source/PerforceSourceControl/Private/PerforceConnection.cpp
+++ b/Engine/Plugins/Developer/PerforceSourceControl/Source/PerforceSourceControl/Private/PerforceConnection.cpp
@@ -626,22 +626,19 @@ bool FPerforceConnection::EnsureValidConnection(FString& InOutServerName, FStrin
 		}
 	}
 
-	const bool bRequireWorkspace = EnumHasAllFlags(Options, EConnectionOptions::WorkspaceOptional);
+	const bool bRequireWorkspace = !EnumHasAllFlags(Options, EConnectionOptions::WorkspaceOptional);
 
 	// Try to auto detect the client if none were specified and we require one
-	if (bConnectionOK)
+	if (bConnectionOK && bRequireWorkspace && NewClientSpecName.IsEmpty())
 	{
-		if ( bRequireWorkspace && NewClientSpecName.IsEmpty() )
-		{
-			FPerforceConnectionInfo AutoCredentials = InConnectionInfo;
-			AutoCredentials.Port = TO_TCHAR(TestP4.GetPort().Text(), bIsUnicodeServer);
-			AutoCredentials.UserName = TO_TCHAR(TestP4.GetUser().Text(), bIsUnicodeServer);
+		FPerforceConnectionInfo AutoCredentials = InConnectionInfo;
+		AutoCredentials.Port = TO_TCHAR(TestP4.GetPort().Text(), bIsUnicodeServer);
+		AutoCredentials.UserName = TO_TCHAR(TestP4.GetUser().Text(), bIsUnicodeServer);
 
-		 	bConnectionOK = FPerforceConnection::AutoDetectWorkspace(AutoCredentials, SCCProvider, NewClientSpecName);
-			if ( bConnectionOK )
-			{
-				TestP4.SetClient(FROM_TCHAR(*NewClientSpecName, bIsUnicodeServer));
-			}
+		bConnectionOK = FPerforceConnection::AutoDetectWorkspace(AutoCredentials, SCCProvider, NewClientSpecName);
+		if (bConnectionOK)
+		{
+			TestP4.SetClient(FROM_TCHAR(*NewClientSpecName, bIsUnicodeServer));
 		}
 	}
 
@@ -679,15 +676,17 @@ bool FPerforceConnection::EnsureValidConnection(FString& InOutServerName, FStrin
 	}
 
 	//if never specified, take the default connection values
-	if (NewServerName.Len() == 0)
+	if (NewServerName.IsEmpty())
 	{
 		NewServerName = TO_TCHAR(TestP4.GetPort().Text(), bIsUnicodeServer);
 	}
-	if (NewUserName.Len() == 0)
+
+	if (NewUserName.IsEmpty())
 	{
 		NewUserName = TO_TCHAR(TestP4.GetUser().Text(), bIsUnicodeServer);
 	}
-	if (NewClientSpecName.Len() == 0)
+
+	if (NewClientSpecName.IsEmpty() && bRequireWorkspace)
 	{
 		NewClientSpecName = TO_TCHAR(TestP4.GetClient().Text(), bIsUnicodeServer);
 		if (NewClientSpecName == TO_TCHAR(TestP4.GetHost().Text(), bIsUnicodeServer))
diff --git a/Engine/Plugins/Developer/PerforceSourceControl/Source/PerforceSourceControl/Private/PerforceSourceControlOperations.cpp b/Engine/Plugins/Developer/PerforceSourceControl/Source/PerforceSourceControl/Private/PerforceSourceControlOperations.cpp
--- a/Engine/Plugins/Developer/PerforceSourceControl/Source/PerforceSourceControl/Private/PerforceSourceControlOperations.cpp
+++ b/Engine/Plugins/Developer/PerforceSourceControl/Source/PerforceSourceControl/Private/PerforceSourceControlOperations.cpp
@@ -337,6 +337,17 @@ bool FPerforceConnectWorker::Execute(FPerforceSourceControlCommand& InCommand)
 {
 	TRACE_CPUPROFILER_EVENT_SCOPE(FPerforceConnectWorker::Execute);
 
+	if (InCommand.ConnectionInfo.Workspace.IsEmpty())
+	{
+		// Running FConnect for perforce does not actually connect to the server but
+		// validates that the clientspec for the connection is correct and working.
+		// So if the connection does not have a valid clientspec then we might as well
+		// early out at this point and return true.
+
+		InCommand.bCommandSuccessful = true;
+		return InCommand.bCommandSuccessful;
+	}
+
 	FScopedPerforceConnection ScopedConnection(InCommand);
 	if (!InCommand.IsCanceled() && ScopedConnection.IsValid())
 	{
diff --git a/Engine/Plugins/Developer/PerforceSourceControl/Source/PerforceSourceControl/Private/PerforceSourceControlProvider.cpp b/Engine/Plugins/Developer/PerforceSourceControl/Source/PerforceSourceControl/Private/PerforceSourceControlProvider.cpp
--- a/Engine/Plugins/Developer/PerforceSourceControl/Source/PerforceSourceControl/Private/PerforceSourceControlProvider.cpp
+++ b/Engine/Plugins/Developer/PerforceSourceControl/Source/PerforceSourceControl/Private/PerforceSourceControlProvider.cpp
@@ -726,12 +726,11 @@ ECommandResult::Type FPerforceSourceControlProvider::SwitchWorkspace(FStringView
 	FString UserName = P4Settings.GetUserName();
 	FString ClientSpecName = FString(NewWorkspaceName);
 
-	if (FPerforceConnection::EnsureValidConnection(PortName, UserName, ClientSpecName, P4Settings.GetConnectionInfo(), *this, EConnectionOptions::None))
+	if (FPerforceConnection::EnsureValidConnection(PortName, UserName, ClientSpecName, P4Settings.GetConnectionInfo(), *this, EConnectionOptions::WorkspaceOptional))
 	{
 		P4Settings.SetPort(PortName);
 		P4Settings.SetUserName(UserName);
 		P4Settings.SetWorkspace(ClientSpecName);
-		check(NewWorkspaceName == ClientSpecName);
 
 		bServerAvailable = true;
 
