Backing out 21008487, submitted to the wrong branch

diff --git a/Engine/Plugins/Developer/PerforceSourceControl/Source/PerforceSourceControl/Private/PerforceConnection.cpp b/Engine/Plugins/Developer/PerforceSourceControl/Source/PerforceSourceControl/Private/PerforceConnection.cpp
--- a/Engine/Plugins/Developer/PerforceSourceControl/Source/PerforceSourceControl/Private/PerforceConnection.cpp
+++ b/Engine/Plugins/Developer/PerforceSourceControl/Source/PerforceSourceControl/Private/PerforceConnection.cpp
@@ -626,19 +626,22 @@ bool FPerforceConnection::EnsureValidConnection(FString& InOutServerName, FStrin
 		}
 	}
 
-	const bool bRequireWorkspace = !EnumHasAllFlags(Options, EConnectionOptions::WorkspaceOptional);
+	const bool bRequireWorkspace = EnumHasAllFlags(Options, EConnectionOptions::WorkspaceOptional);
 
 	// Try to auto detect the client if none were specified and we require one
-	if (bConnectionOK && bRequireWorkspace && NewClientSpecName.IsEmpty())
+	if (bConnectionOK)
 	{
-		FPerforceConnectionInfo AutoCredentials = InConnectionInfo;
-		AutoCredentials.Port = TO_TCHAR(TestP4.GetPort().Text(), bIsUnicodeServer);
-		AutoCredentials.UserName = TO_TCHAR(TestP4.GetUser().Text(), bIsUnicodeServer);
-
-		bConnectionOK = FPerforceConnection::AutoDetectWorkspace(AutoCredentials, SCCProvider, NewClientSpecName);
-		if (bConnectionOK)
+		if ( bRequireWorkspace && NewClientSpecName.IsEmpty() )
 		{
-			TestP4.SetClient(FROM_TCHAR(*NewClientSpecName, bIsUnicodeServer));
+			FPerforceConnectionInfo AutoCredentials = InConnectionInfo;
+			AutoCredentials.Port = TO_TCHAR(TestP4.GetPort().Text(), bIsUnicodeServer);
+			AutoCredentials.UserName = TO_TCHAR(TestP4.GetUser().Text(), bIsUnicodeServer);
+
+		 	bConnectionOK = FPerforceConnection::AutoDetectWorkspace(AutoCredentials, SCCProvider, NewClientSpecName);
+			if ( bConnectionOK )
+			{
+				TestP4.SetClient(FROM_TCHAR(*NewClientSpecName, bIsUnicodeServer));
+			}
 		}
 	}
 
@@ -676,17 +679,15 @@ bool FPerforceConnection::EnsureValidConnection(FString& InOutServerName, FStrin
 	}
 
 	//if never specified, take the default connection values
-	if (NewServerName.IsEmpty())
+	if (NewServerName.Len() == 0)
 	{
 		NewServerName = TO_TCHAR(TestP4.GetPort().Text(), bIsUnicodeServer);
 	}
-
-	if (NewUserName.IsEmpty())
+	if (NewUserName.Len() == 0)
 	{
 		NewUserName = TO_TCHAR(TestP4.GetUser().Text(), bIsUnicodeServer);
 	}
-
-	if (NewClientSpecName.IsEmpty() && bRequireWorkspace)
+	if (NewClientSpecName.Len() == 0)
 	{
 		NewClientSpecName = TO_TCHAR(TestP4.GetClient().Text(), bIsUnicodeServer);
 		if (NewClientSpecName == TO_TCHAR(TestP4.GetHost().Text(), bIsUnicodeServer))
diff --git a/Engine/Plugins/Developer/PerforceSourceControl/Source/PerforceSourceControl/Private/PerforceSourceControlOperations.cpp b/Engine/Plugins/Developer/PerforceSourceControl/Source/PerforceSourceControl/Private/PerforceSourceControlOperations.cpp
--- a/Engine/Plugins/Developer/PerforceSourceControl/Source/PerforceSourceControl/Private/PerforceSourceControlOperations.cpp
+++ b/Engine/Plugins/Developer/PerforceSourceControl/Source/PerforceSourceControl/Private/PerforceSourceControlOperations.cpp
@@ -337,17 +337,6 @@ bool FPerforceConnectWorker::Execute(FPerforceSourceControlCommand& InCommand)
 {
 	TRACE_CPUPROFILER_EVENT_SCOPE(FPerforceConnectWorker::Execute);
 
-	if (InCommand.ConnectionInfo.Workspace.IsEmpty())
-	{
-		// Running FConnect for perforce does not actually connect to the server but
-		// validates that the clientspec for the connection is correct and working.
-		// So if the connection does not have a valid clientspec then we might as well
-		// early out at this point and return true.
-
-		InCommand.bCommandSuccessful = true;
-		return InCommand.bCommandSuccessful;
-	}
-
 	FScopedPerforceConnection ScopedConnection(InCommand);
 	if (!InCommand.IsCanceled() && ScopedConnection.IsValid())
 	{
diff --git a/Engine/Plugins/Developer/PerforceSourceControl/Source/PerforceSourceControl/Private/PerforceSourceControlProvider.cpp b/Engine/Plugins/Developer/PerforceSourceControl/Source/PerforceSourceControl/Private/PerforceSourceControlProvider.cpp
--- a/Engine/Plugins/Developer/PerforceSourceControl/Source/PerforceSourceControl/Private/PerforceSourceControlProvider.cpp
+++ b/Engine/Plugins/Developer/PerforceSourceControl/Source/PerforceSourceControl/Private/PerforceSourceControlProvider.cpp
@@ -726,11 +726,12 @@ ECommandResult::Type FPerforceSourceControlProvider::SwitchWorkspace(FStringView
 	FString UserName = P4Settings.GetUserName();
 	FString ClientSpecName = FString(NewWorkspaceName);
 
-	if (FPerforceConnection::EnsureValidConnection(PortName, UserName, ClientSpecName, P4Settings.GetConnectionInfo(), *this, EConnectionOptions::WorkspaceOptional))
+	if (FPerforceConnection::EnsureValidConnection(PortName, UserName, ClientSpecName, P4Settings.GetConnectionInfo(), *this, EConnectionOptions::None))
 	{
 		P4Settings.SetPort(PortName);
 		P4Settings.SetUserName(UserName);
 		P4Settings.SetWorkspace(ClientSpecName);
+		check(NewWorkspaceName == ClientSpecName);
 
 		bServerAvailable = true;
 
