Duplicating assets should keep a perforce relationship to the original file

diff --git a/Engine/Plugins/Developer/PerforceSourceControl/Source/PerforceSourceControl/Private/PerforceSourceControlOperations.cpp b/Engine/Plugins/Developer/PerforceSourceControl/Source/PerforceSourceControl/Private/PerforceSourceControlOperations.cpp
--- a/Engine/Plugins/Developer/PerforceSourceControl/Source/PerforceSourceControl/Private/PerforceSourceControlOperations.cpp
+++ b/Engine/Plugins/Developer/PerforceSourceControl/Source/PerforceSourceControl/Private/PerforceSourceControlOperations.cpp
@@ -7,6 +7,7 @@
 #include "Algo/IndexOf.h"
 #include "Algo/Transform.h"
 #include "HAL/FileManager.h"
+#include "HAL/IConsoleManager.h"
 #include "ISourceControlModule.h"
 #include "Misc/EngineVersion.h"
 #include "Misc/Paths.h"
@@ -29,6 +30,14 @@ namespace
 DECLARE_DELEGATE_RetVal_OneParam(FPerforceSourceControlWorkerRef, FGetPerforceSourceControlWorker, FPerforceSourceControlProvider&)
 
 static TMap<FName, FGetPerforceSourceControlWorker> WorkersMap;
+
+bool GAlwaysBranchFilesOnCopy = false;
+static FAutoConsoleVariableRef CVarAlwaysBranchFilesOnCopy(
+	TEXT("p4.AlwaysBranchFilesOnCopy"),
+	GAlwaysBranchFilesOnCopy,
+	TEXT("Use legacy behavior of always branching a file in perforce when copying."),
+	ECVF_Default
+);
 }
 
 template<typename Type>
@@ -2313,22 +2322,33 @@ bool FPerforceCopyWorker::Execute(FPerforceSourceControlCommand& InCommand)
 
 		AppendChangelistParameter(GetSCCProvider(), Parameters);
 
-		Parameters.Append(InCommand.Files);
-		Parameters.Add(DestinationPath);
-
 		FP4RecordSet Records;
-		InCommand.bCommandSuccessful = Connection.RunCommand(TEXT("integrate"), Parameters, Records, InCommand.ResultInfo.ErrorMessages, FOnIsCancelled::CreateRaw(&InCommand, &FPerforceSourceControlCommand::IsCanceled), InCommand.bConnectionDropped);
+		if (Operation->CopyMethod == FCopy::ECopyMethod::Branch || GAlwaysBranchFilesOnCopy)
+		{
+			Parameters.Append(InCommand.Files);
+			Parameters.Add(DestinationPath);
 
-		// We now need to do a p4 resolve.
-		// This is because when we copy a file in the Editor, we first make the copy on disk before attempting to branch. This causes a conflict in P4's eyes.
-		// We must do this to prevent the asset registry from picking up what it thinks is a newly-added file (which would be created by the p4 integrate command)
-		// and then the package system getting very confused about where to save the now-duplicated assets.
-		if(InCommand.bCommandSuccessful)
+			InCommand.bCommandSuccessful = Connection.RunCommand(TEXT("integrate"), Parameters, Records, InCommand.ResultInfo.ErrorMessages, FOnIsCancelled::CreateRaw(&InCommand, &FPerforceSourceControlCommand::IsCanceled), InCommand.bConnectionDropped);
+
+			// We now need to do a p4 resolve.
+			// This is because when we copy a file in the Editor, we first make the copy on disk before attempting to branch. This causes a conflict in P4's eyes.
+			// We must do this to prevent the asset registry from picking up what it thinks is a newly-added file (which would be created by the p4 integrate command)
+			// and then the package system getting very confused about where to save the now-duplicated assets.
+			if (InCommand.bCommandSuccessful)
+			{
+				TArray<FString> ResolveParameters;
+				ResolveParameters.Add(TEXT("-ay"));	// 'accept yours'
+				ResolveParameters.Add(MoveTemp(DestinationPath));
+				InCommand.bCommandSuccessful = Connection.RunCommand(TEXT("resolve"), ResolveParameters, Records, InCommand.ResultInfo.ErrorMessages, FOnIsCancelled::CreateRaw(&InCommand, &FPerforceSourceControlCommand::IsCanceled), InCommand.bConnectionDropped);
+			}
+		}
+		else
 		{
-			TArray<FString> ResolveParameters;
-			ResolveParameters.Add(TEXT("-ay"));	// 'accept yours'
-			ResolveParameters.Add(MoveTemp(DestinationPath));
-			InCommand.bCommandSuccessful = Connection.RunCommand(TEXT("resolve"), ResolveParameters, Records, InCommand.ResultInfo.ErrorMessages, FOnIsCancelled::CreateRaw(&InCommand, &FPerforceSourceControlCommand::IsCanceled), InCommand.bConnectionDropped);
+			Parameters.Add(MoveTemp(DestinationPath));
+
+			check(Operation->CopyMethod == FCopy::ECopyMethod::Add);
+
+			InCommand.bCommandSuccessful = Connection.RunCommand(TEXT("add"), Parameters, Records, InCommand.ResultInfo.ErrorMessages, FOnIsCancelled::CreateRaw(&InCommand, &FPerforceSourceControlCommand::IsCanceled), InCommand.bConnectionDropped);
 		}
 	}
 	return InCommand.bCommandSuccessful;
diff --git a/Engine/Source/Developer/AssetTools/Private/AssetTools.cpp b/Engine/Source/Developer/AssetTools/Private/AssetTools.cpp
--- a/Engine/Source/Developer/AssetTools/Private/AssetTools.cpp
+++ b/Engine/Source/Developer/AssetTools/Private/AssetTools.cpp
@@ -699,7 +699,7 @@ UObject* UAssetToolsImpl::PerformDuplicateAsset(const FString& AssetName, const
 			FEditorFileUtils::PromptForCheckoutAndSave(OutermostPackagesToSave, bCheckDirty, bPromptToSave);
 
 			// now attempt to branch, we can do this now as we should have a file on disk
-			SourceControlHelpers::BranchPackage(NewObject->GetOutermost(), OriginalObject->GetOutermost());
+			SourceControlHelpers::CopyPackage(NewObject->GetOutermost(), OriginalObject->GetOutermost());
 		}
 
 		// analytics create record
diff --git a/Engine/Source/Developer/AssetTools/Private/AssetViewUtils.cpp b/Engine/Source/Developer/AssetTools/Private/AssetViewUtils.cpp
--- a/Engine/Source/Developer/AssetTools/Private/AssetViewUtils.cpp
+++ b/Engine/Source/Developer/AssetTools/Private/AssetViewUtils.cpp
@@ -272,7 +272,7 @@ void AssetViewUtils::CopyAssets(const TArray<UObject*>& Assets, const FString& D
 		{
 			UObject* SourceAsset = Assets[ObjectIndex];
 			UObject* DestAsset = NewObjects[ObjectIndex];
-			SourceControlHelpers::BranchPackage(DestAsset->GetOutermost(), SourceAsset->GetOutermost());
+			SourceControlHelpers::CopyPackage(DestAsset->GetOutermost(), SourceAsset->GetOutermost());
 		}
 	}
 }
diff --git a/Engine/Source/Developer/SourceControl/Private/SourceControlHelpers.cpp b/Engine/Source/Developer/SourceControl/Private/SourceControlHelpers.cpp
--- a/Engine/Source/Developer/SourceControl/Private/SourceControlHelpers.cpp
+++ b/Engine/Source/Developer/SourceControl/Private/SourceControlHelpers.cpp
@@ -1521,26 +1521,38 @@ bool USourceControlHelpers::CopyFileUnderSourceControl( const FString& InDestFil
 	return CheckoutOrMarkForAdd(InDestFile, InFileDescription, FOnPostCheckOut::CreateStatic(&Local::CopyFile, InSourceFile), OutFailReason);
 }
 
-
-bool USourceControlHelpers::BranchPackage( UPackage* DestPackage, UPackage* SourcePackage, EStateCacheUsage::Type StateCacheUsage )
+namespace
 {
-	if(ISourceControlModule::Get().IsEnabled())
+	bool CopyPackage_Internal(UPackage* DestPackage, UPackage* SourcePackage, FCopy::ECopyMethod CopyMethod, EStateCacheUsage::Type StateCacheUsage)
 	{
-		ISourceControlProvider& SourceControlProvider = ISourceControlModule::Get().GetProvider();
-
-		const FString SourceFilename = PackageFilename(SourcePackage);
-		const FString DestFilename = PackageFilename(DestPackage);
-		FSourceControlStatePtr SourceControlState = SourceControlProvider.GetState(SourceFilename, StateCacheUsage);
-		if(SourceControlState.IsValid() && SourceControlState->IsSourceControlled())
+		if (ISourceControlModule::Get().IsEnabled())
 		{
-			TSharedRef<FCopy, ESPMode::ThreadSafe> CopyOperation = ISourceControlOperation::Create<FCopy>();
-			CopyOperation->SetDestination(DestFilename);
-			
-			return (SourceControlProvider.Execute(CopyOperation, SourceFilename) == ECommandResult::Succeeded);
+			ISourceControlProvider& SourceControlProvider = ISourceControlModule::Get().GetProvider();
+
+			const FString SourceFilename = USourceControlHelpers::PackageFilename(SourcePackage);
+			FSourceControlStatePtr SourceControlState = SourceControlProvider.GetState(SourceFilename, StateCacheUsage);
+			if (SourceControlState.IsValid() && SourceControlState->IsSourceControlled())
+			{
+				TSharedRef<FCopy, ESPMode::ThreadSafe> CopyOperation = ISourceControlOperation::Create<FCopy>();
+				CopyOperation->SetDestination(USourceControlHelpers::PackageFilename(DestPackage));
+				CopyOperation->CopyMethod = CopyMethod;
+
+				return (SourceControlProvider.Execute(CopyOperation, SourceFilename) == ECommandResult::Succeeded);
+			}
 		}
+
+		return false;
 	}
+}
 
-	return false;
+bool USourceControlHelpers::BranchPackage(UPackage* DestPackage, UPackage* SourcePackage, EStateCacheUsage::Type StateCacheUsage)
+{
+	return CopyPackage_Internal(DestPackage, SourcePackage, FCopy::ECopyMethod::Branch, StateCacheUsage);
+}
+
+bool USourceControlHelpers::CopyPackage(UPackage* DestPackage, UPackage* SourcePackage, EStateCacheUsage::Type StateCacheUsage)
+{
+	return CopyPackage_Internal(DestPackage, SourcePackage, FCopy::ECopyMethod::Add, StateCacheUsage);
 }
 
 
diff --git a/Engine/Source/Developer/SourceControl/Public/SourceControlHelpers.h b/Engine/Source/Developer/SourceControl/Public/SourceControlHelpers.h
--- a/Engine/Source/Developer/SourceControl/Public/SourceControlHelpers.h
+++ b/Engine/Source/Developer/SourceControl/Public/SourceControlHelpers.h
@@ -518,7 +518,8 @@ public:
 	static bool AnnotateFile(ISourceControlProvider& InProvider, int32 InCheckInIdentifier, const FString& InFile, TArray<FAnnotationLine>& OutLines);
 
 	/**
-	 * Helper function to branch/integrate packages from one location to another
+	 * Helper function to branch/integrate packages from one location to another maintaining
+	 * a relationship between the files in source control (when possible)
 	 * @param	DestPackage			The destination package
 	 * @param	SourcePackage		The source package
 	 * @Param	StateCacheUsage		Whether to use the source control state cache
@@ -526,6 +527,15 @@ public:
 	 */
 	static bool BranchPackage(UPackage* DestPackage, UPackage* SourcePackage, EStateCacheUsage::Type StateCacheUsage = EStateCacheUsage::ForceUpdate);
 
+	/**
+	 * Helper function to copy a package from one location to another
+	 * @param	DestPackage             The destination package
+	 * @param	SourcePackage           The source package
+	 * @Param	StateCacheUsage         Whether to use the source control state cache
+	 * @return true if the file packages were successfully branched.
+	 */
+	static bool CopyPackage(UPackage* DestPackage, UPackage* SourcePackage, EStateCacheUsage::Type StateCacheUsage = EStateCacheUsage::ForceUpdate);
+
 	/**
 	 * Helper function to get the ini filename for storing source control settings
 	 * @return the filename
diff --git a/Engine/Source/Developer/SourceControl/Public/SourceControlOperations.h b/Engine/Source/Developer/SourceControl/Public/SourceControlOperations.h
--- a/Engine/Source/Developer/SourceControl/Public/SourceControlOperations.h
+++ b/Engine/Source/Developer/SourceControl/Public/SourceControlOperations.h
@@ -506,6 +506,15 @@ public:
 		return Destination;
 	}
 
+	enum class ECopyMethod
+	{
+		Branch,  // The new file is branched from the original file
+		Add      // The new file has no relation to the original file
+	};
+
+	/** Whether a relationship to the original file should be maintained */
+	ECopyMethod CopyMethod = ECopyMethod::Branch; 
+
 protected:
 	/** Destination path of the copy operation */
 	FString Destination;
