Perforce Source Control: always print perforce command execution time when it takes over 100ms.

#rnx

diff --git a/Engine/Plugins/Developer/PerforceSourceControl/Source/PerforceSourceControl/Private/PerforceConnection.cpp b/Engine/Plugins/Developer/PerforceSourceControl/Source/PerforceSourceControl/Private/PerforceConnection.cpp
--- a/Engine/Plugins/Developer/PerforceSourceControl/Source/PerforceSourceControl/Private/PerforceConnection.cpp
+++ b/Engine/Plugins/Developer/PerforceSourceControl/Source/PerforceSourceControl/Private/PerforceConnection.cpp
@@ -906,7 +906,8 @@ bool FPerforceConnection::RunCommand(	const FString& InCommand, const TArray<FSt
 
 	if (bLogCommandDetails)
 	{
-		UE_LOG( LogSourceControl, VeryVerbose, TEXT("P4 execution time: %0.4f seconds. Command: %s"), FPlatformTime::Seconds() - SCCStartTime, *FullCommand );
+		const double ExecutionTime = FPlatformTime::Seconds() - SCCStartTime;
+		UE_CLOG((LogSourceControl.GetVerbosity() == ELogVerbosity::VeryVerbose) || (ExecutionTime >= 0.1), LogSourceControl, Warning, TEXT("P4 execution time: %0.4f seconds. Command: %s"), ExecutionTime, *FullCommand);
 	}
 
 	return OutRecordSet.Num() > 0;
