Downgrade a warning to a log.

#rnx

diff --git a/Engine/Plugins/Developer/PerforceSourceControl/Source/PerforceSourceControl/Private/PerforceConnection.cpp b/Engine/Plugins/Developer/PerforceSourceControl/Source/PerforceSourceControl/Private/PerforceConnection.cpp
--- a/Engine/Plugins/Developer/PerforceSourceControl/Source/PerforceSourceControl/Private/PerforceConnection.cpp
+++ b/Engine/Plugins/Developer/PerforceSourceControl/Source/PerforceSourceControl/Private/PerforceConnection.cpp
@@ -907,7 +907,7 @@ bool FPerforceConnection::RunCommand(	const FString& InCommand, const TArray<FSt
 	if (bLogCommandDetails)
 	{
 		const double ExecutionTime = FPlatformTime::Seconds() - SCCStartTime;
-		UE_CLOG((LogSourceControl.GetVerbosity() == ELogVerbosity::VeryVerbose) || (ExecutionTime >= 0.1), LogSourceControl, Warning, TEXT("P4 execution time: %0.4f seconds. Command: %s"), ExecutionTime, *FullCommand);
+		UE_CLOG((LogSourceControl.GetVerbosity() == ELogVerbosity::VeryVerbose) || (ExecutionTime >= 0.1), LogSourceControl, Log, TEXT("P4 execution time: %0.4f seconds. Command: %s"), ExecutionTime, *FullCommand);
 	}
 
 	return OutRecordSet.Num() > 0;
