Do not try validating the perforce connection when switching to not using a workspace.

#rnx

- If we are switching workspaces then we already know that the existing perforce settings are correct.
- The only reason we have to validate them is if we are switching to a different workspace and need to make sure that it exists.
- If we are just removing the workspace name entirely (which is valid for a number of operations) then we have nothing to check and we can just remove the workspace name from our p4 settings structure.
- This saves issuing a few pointless p4 commands.

diff --git a/Engine/Plugins/Developer/PerforceSourceControl/Source/PerforceSourceControl/Private/PerforceSourceControlProvider.cpp b/Engine/Plugins/Developer/PerforceSourceControl/Source/PerforceSourceControl/Private/PerforceSourceControlProvider.cpp
--- a/Engine/Plugins/Developer/PerforceSourceControl/Source/PerforceSourceControl/Private/PerforceSourceControlProvider.cpp
+++ b/Engine/Plugins/Developer/PerforceSourceControl/Source/PerforceSourceControl/Private/PerforceSourceControlProvider.cpp
@@ -721,39 +721,48 @@ ECommandResult::Type FPerforceSourceControlProvider::SwitchWorkspace(FStringView
 	FPerforceSourceControlSettings& P4Settings = AccessSettings(); 
 	
 	const FString OldWorkspaceName = P4Settings.GetWorkspace();
-		
-	FString PortName = P4Settings.GetPort();
-	FString UserName = P4Settings.GetUserName();
-	FString ClientSpecName = FString(NewWorkspaceName);
+	FString WorkspaceName = FString(NewWorkspaceName);
 
-	if (FPerforceConnection::EnsureValidConnection(PortName, UserName, ClientSpecName, P4Settings.GetConnectionInfo(), *this, EConnectionOptions::WorkspaceOptional))
+	if (!NewWorkspaceName.IsEmpty())
 	{
-		P4Settings.SetPort(PortName);
-		P4Settings.SetUserName(UserName);
-		P4Settings.SetWorkspace(ClientSpecName);
+		FString PortName = P4Settings.GetPort();
+		FString UserName = P4Settings.GetUserName();
+		
+		if (FPerforceConnection::EnsureValidConnection(PortName, UserName, WorkspaceName, P4Settings.GetConnectionInfo(), *this, EConnectionOptions::WorkspaceOptional))
+		{
+			P4Settings.SetPort(PortName);
+			P4Settings.SetUserName(UserName);
+			P4Settings.SetWorkspace(WorkspaceName);
 
-		bServerAvailable = true;
+			bServerAvailable = true;
 
-		if (OutOldWorkspaceName != nullptr)
-		{
-			*OutOldWorkspaceName = OldWorkspaceName;
+			if (OutOldWorkspaceName != nullptr)
+			{
+				*OutOldWorkspaceName = OldWorkspaceName;
+			}	
 		}
+		else
+		{
+			FText Message = FText::Format(LOCTEXT("Perforce_ConnectionFailed", "Failed to re-establish the connection after switching to workspace {0}"),
+				FText::FromString(FString(NewWorkspaceName)));
+			OutResultInfo.ErrorMessages.Add(Message);
 
-		UE_LOG(LogSourceControl, Log, TEXT("Switched workspaces from '%s' to '%s%'"),  *OldWorkspaceName, *ClientSpecName);
+			// The connection didn't work so we should try to restore the old workspace name
+			P4Settings.SetWorkspace(OldWorkspaceName);
 
-		return ECommandResult::Succeeded;
+			return ECommandResult::Failed;
+		}
 	}
 	else
 	{
-		FText Message = FText::Format(	LOCTEXT("Perforce_ConnectionFailed", "Failed to re-establish the connection after switching to workspace {0}"),
-										FText::FromString(FString(NewWorkspaceName)));
-		OutResultInfo.ErrorMessages.Add(Message);
-
-		// The connection didn't work so we should try to restore the old workspace name
-		P4Settings.SetWorkspace(OldWorkspaceName);
-		
-		return ECommandResult::Failed;
+		// If we are just removing the workspace name then we don't need to ensure a valid connection, we can just go ahead and remove the name 
+		// and continue with the existing settings.
+		P4Settings.SetWorkspace(WorkspaceName);
 	}
+
+	UE_LOG(LogSourceControl, Log, TEXT("Switched workspaces from '%s' to '%s%'"), *OldWorkspaceName, *WorkspaceName);
+
+	return ECommandResult::Succeeded;
 }
 
 #if SOURCE_CONTROL_WITH_SLATE
