Removed the creation of hundreds of unneccesary temporary strings while parsing P4 results.

diff --git a/Engine/Plugins/Developer/PerforceSourceControl/Source/PerforceSourceControl/Private/PerforceConnection.cpp b/Engine/Plugins/Developer/PerforceSourceControl/Source/PerforceSourceControl/Private/PerforceConnection.cpp
--- a/Engine/Plugins/Developer/PerforceSourceControl/Source/PerforceSourceControl/Private/PerforceConnection.cpp
+++ b/Engine/Plugins/Developer/PerforceSourceControl/Source/PerforceSourceControl/Private/PerforceConnection.cpp
@@ -32,6 +32,9 @@ namespace
 	ENUM_CLASS_FLAGS(EP4ClientUserFlags);
 }
 
+const FString FP4Record::EmptyStr;
+
+
 /** Custom ClientUser class for handling results and errors from Perforce commands */
 class FP4ClientUser : public ClientUser
 {
@@ -63,7 +66,7 @@ public:
 		// time.
 		if (IsCollectingData() && DataBuffer.GetAllocatedSize() == 0)
 		{
-			const FString SizeAsString = Record(TEXT("fileSize"));
+			const FString& SizeAsString = Record(TEXT("fileSize"));
 			if (!SizeAsString.IsEmpty())
 			{
 				const int64 Size = FCString::Atoi64(*SizeAsString);
@@ -743,12 +746,11 @@ bool FPerforceConnection::GetWorkspaceList(const FPerforceConnectionInfo& InConn
 				LocalHostName = LocalHostName.ToLower();
 			}
 
-			for (int32 Index = 0; Index < Records.Num(); ++Index)
+			for (const FP4Record& ClientRecord : Records)
 			{
-				const FP4Record& ClientRecord = Records[Index];
-				FString ClientName = ClientRecord("client");
-				FString HostName = ClientRecord("Host");
-				FString ClientRootPath = ClientRecord("Root").ToLower();
+				const FString& ClientName = ClientRecord(TEXT("client"));
+				const FString& HostName = ClientRecord(TEXT("Host"));
+				FString ClientRootPath = ClientRecord(TEXT("Root")).ToLower();
 
 				//this clientspec has to be meant for this machine ( "" hostnames mean any host can use ths clientspec in p4 land)
 				bool bHostNameMatches = (LocalHostName == HostName.ToLower());
diff --git a/Engine/Plugins/Developer/PerforceSourceControl/Source/PerforceSourceControl/Private/PerforceConnection.h b/Engine/Plugins/Developer/PerforceSourceControl/Source/PerforceSourceControl/Private/PerforceConnection.h
--- a/Engine/Plugins/Developer/PerforceSourceControl/Source/PerforceSourceControl/Private/PerforceConnection.h
+++ b/Engine/Plugins/Developer/PerforceSourceControl/Source/PerforceSourceControl/Private/PerforceConnection.h
@@ -10,17 +10,31 @@
 class FPerforceSourceControlProvider;
 
 /** A map containing result of running Perforce command */
-class FP4Record : public TMap<FString, FString >
+class FP4Record : public TMap<FString, FString>
 {
 public:
-
 	virtual ~FP4Record() {}
 
-	FString operator () (const FString& Key) const
+	const FString& operator()(const FString& Key) const
+	{
+		if (const FString* Found = Find(Key))
+		{
+			return *Found;
+		}
+		return EmptyStr;
+	}
+
+	const FString& operator()(const TCHAR* Key) const
 	{
-		const FString* Found = Find(Key);
-		return Found ? *Found : TEXT("");
+		if (const FString* Found = FindByHash(GetTypeHash(Key), Key))
+		{
+			return *Found;
+		}
+		return EmptyStr;
 	}
+
+private:
+	static const FString EmptyStr;
 };
 
 typedef TArray<FP4Record> FP4RecordSet;
diff --git a/Engine/Plugins/Developer/PerforceSourceControl/Source/PerforceSourceControl/Private/PerforceSourceControlLabel.cpp b/Engine/Plugins/Developer/PerforceSourceControl/Source/PerforceSourceControl/Private/PerforceSourceControlLabel.cpp
--- a/Engine/Plugins/Developer/PerforceSourceControl/Source/PerforceSourceControl/Private/PerforceSourceControlLabel.cpp
+++ b/Engine/Plugins/Developer/PerforceSourceControl/Source/PerforceSourceControl/Private/PerforceSourceControlLabel.cpp
@@ -19,19 +19,18 @@ const FString& FPerforceSourceControlLabel::GetName() const
 static void ParseFilesResults(FPerforceSourceControlProvider& SCCProvider, const FP4RecordSet& InRecords, const FString& InClientRoot, TArray< TSharedRef<ISourceControlRevision, ESPMode::ThreadSafe> >& OutRevisions)
 {
 	// Iterate over each record found as a result of the command, parsing it for relevant information
-	for (int32 Index = 0; Index < InRecords.Num(); ++Index)
+	for (const FP4Record& ClientRecord : InRecords)
 	{
-		const FP4Record& ClientRecord = InRecords[Index];
-		FString DepotFile = ClientRecord(TEXT("depotFile"));
-		FString RevisionNumber = ClientRecord(TEXT("rev"));
-		FString Date = ClientRecord(TEXT("time"));
-		FString ChangelistNumber = ClientRecord(TEXT("change"));
-		FString Action = ClientRecord(TEXT("action"));
+		const FString& DepotFile = ClientRecord(TEXT("depotFile"));
+		const FString& RevisionNumber = ClientRecord(TEXT("rev"));
+		const FString& Date = ClientRecord(TEXT("time"));
+		const FString& ChangelistNumber = ClientRecord(TEXT("change"));
+		const FString& Action = ClientRecord(TEXT("action"));
 		check(RevisionNumber.Len() != 0);
 
 		// @todo: this revision is incomplete, but is sufficient for now given the usage of labels to get files, rather
 		// than review revision histories.
-		TSharedRef<FPerforceSourceControlRevision, ESPMode::ThreadSafe> Revision = MakeShareable(new FPerforceSourceControlRevision(SCCProvider));
+		TSharedRef<FPerforceSourceControlRevision, ESPMode::ThreadSafe> Revision = MakeShared<FPerforceSourceControlRevision>(SCCProvider);
 		Revision->FileName = DepotFile;
 		Revision->RevisionNumber = FCString::Atoi(*RevisionNumber);
 		Revision->ChangelistNumber = FCString::Atoi(*ChangelistNumber);
diff --git a/Engine/Plugins/Developer/PerforceSourceControl/Source/PerforceSourceControl/Private/PerforceSourceControlOperations.cpp b/Engine/Plugins/Developer/PerforceSourceControl/Source/PerforceSourceControl/Private/PerforceSourceControlOperations.cpp
--- a/Engine/Plugins/Developer/PerforceSourceControl/Source/PerforceSourceControl/Private/PerforceSourceControlOperations.cpp
+++ b/Engine/Plugins/Developer/PerforceSourceControl/Source/PerforceSourceControl/Private/PerforceSourceControlOperations.cpp
@@ -43,7 +43,7 @@ static FAutoConsoleVariableRef CVarAlwaysBranchFilesOnCopy(
 template<typename Type>
 FPerforceSourceControlWorkerRef InstantiateWorker(FPerforceSourceControlProvider& InSourceControlProvider)
 {
-	return MakeShareable(new Type(InSourceControlProvider));
+	return MakeShared<Type>(InSourceControlProvider);
 }
 
 void IPerforceSourceControlWorker::RegisterWorkers()
@@ -194,11 +194,10 @@ static void ParseRecordSet(const FP4RecordSet& InRecords, TArray<FText>& OutResu
 static void ParseRecordSetForState(const FP4RecordSet& InRecords, TMap<FString, EPerforceState::Type>& OutResults)
 {
 	// Iterate over each record found as a result of the command, parsing it for relevant information
-	for (int32 Index = 0; Index < InRecords.Num(); ++Index)
+	for (const FP4Record& ClientRecord : InRecords)
 	{
-		const FP4Record& ClientRecord = InRecords[Index];
-		FString FileName = ClientRecord(TEXT("clientFile"));
-		FString Action = ClientRecord(TEXT("action"));
+		const FString& FileName = ClientRecord(TEXT("clientFile"));
+		const FString& Action = ClientRecord(TEXT("action"));
 
 		check(FileName.Len());
 		FString FullPath(FileName);
@@ -224,7 +223,7 @@ static void ParseRecordSetForState(const FP4RecordSet& InRecords, TMap<FString,
 			}
 			else if(Action == TEXT("reverted"))
 			{
-				FString OldAction = ClientRecord(TEXT("oldAction"));
+				const FString& OldAction = ClientRecord(TEXT("oldAction"));
 				if(IsAddAction(OldAction))
 				{
 					OutResults.Add(FullPath, EPerforceState::NotInDepot);
@@ -314,7 +313,7 @@ static bool CheckWorkspaceRecordSet(const FP4RecordSet& InRecords, TArray<FText>
 		}
 		else
 		{
-			const FString Client = Record(TEXT("Client"));
+			const FString& Client = Record(TEXT("Client"));
 			OutNotificationText = FText::Format(LOCTEXT("WorkspaceError", "Workspace '{0}' does not map into this project's directory."), FText::FromString(Client));
 			OutErrorMessages.Add(OutNotificationText);
 			OutErrorMessages.Add(LOCTEXT("WorkspaceHelp", "You should set your workspace up to map to a directory at or above the project's directory."));
@@ -458,10 +457,9 @@ FName FPerforceCheckInWorker::GetName() const
 static FText ParseSubmitResults(const FP4RecordSet& InRecords)
 {
 	// Iterate over each record found as a result of the command, parsing it for relevant information
-	for (int32 Index = 0; Index < InRecords.Num(); ++Index)
+	for (const FP4Record& ClientRecord : InRecords)
 	{
-		const FP4Record& ClientRecord = InRecords[Index];
-		const FString SubmittedChange = ClientRecord(TEXT("submittedChange"));
+		const FString& SubmittedChange = ClientRecord(TEXT("submittedChange"));
 		if(SubmittedChange.Len() > 0)
 		{
 			return FText::Format(LOCTEXT("SubmitMessage", "Submitted changelist {0}"), FText::FromString(SubmittedChange));
@@ -958,11 +956,10 @@ FName FPerforceSyncWorker::GetName() const
 static void ParseSyncResults(const FP4RecordSet& InRecords, TMap<FString, EPerforceState::Type>& OutResults)
 {
 	// Iterate over each record found as a result of the command, parsing it for relevant information
-	for (int32 Index = 0; Index < InRecords.Num(); ++Index)
+	for (const FP4Record& ClientRecord : InRecords)
 	{
-		const FP4Record& ClientRecord = InRecords[Index];
-		FString FileName = ClientRecord(TEXT("clientFile"));
-		FString Action = ClientRecord(TEXT("action"));
+		const FString& FileName = ClientRecord(TEXT("clientFile"));
+		const FString& Action = ClientRecord(TEXT("action"));
 
 		check(FileName.Len());
 		FString FullPath(FileName);
@@ -1053,13 +1050,12 @@ bool FPerforceSyncWorker::UpdateStates() const
 static void ParseBranchModificationResults(const FP4RecordSet& InRecords, const TArray<FText>& ErrorMessages, const FString& ContentRoot, TMap<FString, FBranchModification>& BranchModifications)
 {
 	TRACE_CPUPROFILER_EVENT_SCOPE(FPerforce::ParseBranchModificationResults);
-		
-	for (int32 Index = 0; Index < InRecords.Num(); ++Index)
+
+	for (const FP4Record& ClientRecord : InRecords)
 	{
-		const FP4Record& ClientRecord = InRecords[Index];
-		FString DepotFileName = ClientRecord(TEXT("depotFile"));
-		FString ClientFileName = ClientRecord(TEXT("clientFile"));
-		FString HeadAction = ClientRecord(TEXT("headAction"));
+		const FString& DepotFileName = ClientRecord(TEXT("depotFile"));
+		const FString& ClientFileName = ClientRecord(TEXT("clientFile"));
+		const FString& HeadAction = ClientRecord(TEXT("headAction"));
 		int64 HeadModTime = FCString::Atoi64(*ClientRecord(TEXT("headModTime")));
 		int64 HeadTime = FCString::Atoi64(*ClientRecord(TEXT("headTime")));
 		int32 HeadChange = FCString::Atoi(*ClientRecord(TEXT("headChange")));
@@ -1143,9 +1139,8 @@ static void ParseUpdateStatusResults(const FP4RecordSet& InRecords, const TArray
 	TRACE_CPUPROFILER_EVENT_SCOPE(FPerforceUpdateStatusWorker::ParseUpdateStatusResults);
 
 	// Build up a map of any other branch states
-	for (int32 Index = 0; Index < InRecords.Num(); ++Index)
+	for (const FP4Record& ClientRecord : InRecords)
 	{
-		const FP4Record& ClientRecord = InRecords[Index];
 		FString FileName = ClientRecord(TEXT("clientFile"));
 
 		if (FileName.Len())
@@ -1155,8 +1150,8 @@ static void ParseUpdateStatusResults(const FP4RecordSet& InRecords, const TArray
 		}
 
 		// Get the content filename and add to branch states
-		FString DepotFileName = ClientRecord(TEXT("depotFile"));
-		FString OtherOpen = ClientRecord(TEXT("otherOpen"));
+		const FString& DepotFileName = ClientRecord(TEXT("depotFile"));
+		const FString& OtherOpen = ClientRecord(TEXT("otherOpen"));
 
 		FString Branch;
 
@@ -1190,7 +1185,7 @@ static void ParseUpdateStatusResults(const FP4RecordSet& InRecords, const TArray
 			for (int32 OpenIdx = 0; OpenIdx < OtherOpenNum; ++OpenIdx)
 			{
 				const FString OtherOpenRecordKey = FString::Printf(TEXT("otherOpen%d"), OpenIdx);
-				const FString OtherOpenRecordValue = ClientRecord(OtherOpenRecordKey);
+				const FString& OtherOpenRecordValue = ClientRecord(OtherOpenRecordKey);
 
 				int32 AtIndex = OtherOpenRecordValue.Find(TEXT("@"));
 				FString OtherOpenUser = AtIndex == INDEX_NONE ? FString(TEXT("")) : OtherOpenRecordValue.Left(AtIndex);
@@ -1202,24 +1197,22 @@ static void ParseUpdateStatusResults(const FP4RecordSet& InRecords, const TArray
 				}
 			}
 		}
-
 	}
 
 	// Iterate over each record found as a result of the command, parsing it for relevant information
-	for (int32 Index = 0; Index < InRecords.Num(); ++Index)
-	{
-		const FP4Record& ClientRecord = InRecords[Index];
-		FString FileName = ClientRecord(TEXT("clientFile"));
-		FString DepotFileName = ClientRecord(TEXT("depotFile"));
-		FString Changelist = ClientRecord(TEXT("change"));
-		FString HeadRev  = ClientRecord(TEXT("headRev"));
-		FString HaveRev  = ClientRecord(TEXT("haveRev"));
-		FString OtherOpen = ClientRecord(TEXT("otherOpen"));
-		FString OpenType = ClientRecord(TEXT("type"));
-		FString HeadAction = ClientRecord(TEXT("headAction"));
-		FString Action = ClientRecord(TEXT("action"));
-		FString HeadType = ClientRecord(TEXT("headType"));
-		FString Digest = ClientRecord(TEXT("digest"));
+	for (const FP4Record& ClientRecord : InRecords)
+	{
+		const FString& FileName = ClientRecord(TEXT("clientFile"));
+		const FString& DepotFileName = ClientRecord(TEXT("depotFile"));
+		const FString& Changelist = ClientRecord(TEXT("change"));
+		const FString& HeadRev  = ClientRecord(TEXT("headRev"));
+		const FString& HaveRev  = ClientRecord(TEXT("haveRev"));
+		const FString& OtherOpen = ClientRecord(TEXT("otherOpen"));
+		const FString& OpenType = ClientRecord(TEXT("type"));
+		const FString& HeadAction = ClientRecord(TEXT("headAction"));
+		const FString& Action = ClientRecord(TEXT("action"));
+		const FString& HeadType = ClientRecord(TEXT("headType"));
+		const FString& Digest = ClientRecord(TEXT("digest"));
 		const bool bUnresolved = ClientRecord.Contains(TEXT("unresolved"));
 
 		if (!FileName.Len())
@@ -1272,7 +1265,7 @@ static void ParseUpdateStatusResults(const FP4RecordSet& InRecords, const TArray
 			for ( int32 OpenIdx = 0; OpenIdx < OtherOpenNum; ++OpenIdx )
 			{
 				const FString OtherOpenRecordKey = FString::Printf(TEXT("otherOpen%d"), OpenIdx);
-				const FString OtherOpenRecordValue = ClientRecord(OtherOpenRecordKey);
+				const FString& OtherOpenRecordValue = ClientRecord(OtherOpenRecordKey);
 
 				int32 AtIndex = OtherOpenRecordValue.Find(TEXT("@"));
 				FString OtherOpenUser = AtIndex == INDEX_NONE ? FString(TEXT("")) : OtherOpenRecordValue.Left(AtIndex);
@@ -1369,16 +1362,16 @@ static void ParseUpdateStatusResults(const FP4RecordSet& InRecords, const TArray
 					}
 
 					VarName = FString::Printf(TEXT("resolveBaseFile%d"), ResolveActionNumber);
-					FString ResolveBaseFile = ClientRecord(VarName);
+					const FString& ResolveBaseFile = ClientRecord(VarName);
 					VarName = FString::Printf(TEXT("resolveFromFile%d"), ResolveActionNumber);
-					FString ResolveFromFile = ClientRecord(VarName);
+					const FString& ResolveFromFile = ClientRecord(VarName);
 					if(!ensureMsgf( ResolveFromFile == ResolveBaseFile, TEXT("Text cannot resolve %s with %s, we do not support cross file merging"), *ResolveBaseFile, *ResolveFromFile ) )
 					{
 						break;
 					}
 
 					VarName = FString::Printf(TEXT("resolveBaseRev%d"), ResolveActionNumber);
-					FString ResolveBaseRev = ClientRecord(VarName);
+					const FString& ResolveBaseRev = ClientRecord(VarName);
 
 					TTypeFromString<int>::FromString(State.PendingResolveRevNumber, *ResolveBaseRev);
 
@@ -1403,10 +1396,8 @@ static void ParseUpdateStatusResults(const FP4RecordSet& InRecords, const TArray
 	}
 
 	// also see if we can glean anything from the error messages
-	for (int32 Index = 0; Index < ErrorMessages.Num(); ++Index)
+	for (const FText& Error : ErrorMessages)
 	{
-		const FText& Error = ErrorMessages[Index];
-
 		//@todo P4 could be returning localized error messages
 		int32 NoSuchFilePos = Error.ToString().Find(TEXT(" - no such file(s).\n"), ESearchCase::IgnoreCase, ESearchDir::FromStart);
 		if(NoSuchFilePos != INDEX_NONE)
@@ -1437,11 +1428,9 @@ static void ParseUpdateStatusResults(const FP4RecordSet& InRecords, const TArray
 static void ParseOpenedResults(const FP4RecordSet& InRecords, const FString& ClientName, const FString& ClientRoot, TArray<FPerforceSourceControlState>& OutResults)
 {
 	// Iterate over each record found as a result of the command, parsing it for relevant information
-	for (int32 Index = 0; Index < InRecords.Num(); ++Index)
+	for (const FP4Record& ClientRecord : InRecords)
 	{
-		const FP4Record& ClientRecord = InRecords[Index];
-		FString ClientFileName = ClientRecord(TEXT("clientFile"));
-
+		const FString& ClientFileName = ClientRecord(TEXT("clientFile"));
 		check(ClientFileName.Len() > 0);
 
 		// Convert the depot file name to a local file name
@@ -1475,7 +1464,7 @@ static void ParseOpenedResults(const FP4RecordSet& InRecords, const FString& Cli
 		FPerforceSourceControlState& OutState = OutResults.Emplace_GetRef(FullPath);
 		OutState.DepotFilename = ClientRecord(TEXT("depotFile"));
 
-		FString Action = ClientRecord(TEXT("action"));
+		const FString& Action = ClientRecord(TEXT("action"));
 		if (Action.Len() > 0)
 		{
 			if(IsAddAction(Action))
@@ -1492,7 +1481,7 @@ static void ParseOpenedResults(const FP4RecordSet& InRecords, const FString& Cli
 			}
 		}
 
-		FString Changelist = ClientRecord(TEXT("change"));
+		const FString& Changelist = ClientRecord(TEXT("change"));
 		if (Changelist.Len() > 0 && Changelist != TEXT("default"))
 		{
 			OutState.Changelist = FPerforceSourceControlChangelist(FCString::Atoi(*Changelist));
@@ -1502,7 +1491,7 @@ static void ParseOpenedResults(const FP4RecordSet& InRecords, const FString& Cli
 			OutState.Changelist = FPerforceSourceControlChangelist::DefaultChangelist;
 		}
 
-		FString Type = ClientRecord(TEXT("type"));
+		const FString& Type = ClientRecord(TEXT("type"));
 		if (Type.Len() > 0)
 		{
 			OutState.bBinary = Type.Contains(TEXT("binary"));
@@ -1526,11 +1515,10 @@ static void ParseOpenedResults(const FP4RecordSet& InRecords, const FString& Cli
 static void ParseShelvedResults(const FP4RecordSet& InRecords, TMap<FString, EPerforceState::Type>& OutResults)
 {
 	// Iterate over each record found as a result of the command, parsing it for relevant information
-	for (int32 Index = 0; Index < InRecords.Num(); ++Index)
+	for (const FP4Record& Record : InRecords)
 	{
-		const FP4Record& Record = InRecords[Index];
-		FString DepotFileName = Record(TEXT("depotFile"));
-		FString Action = Record(TEXT("action"));
+		const FString& DepotFileName = Record(TEXT("depotFile"));
+		const FString& Action = Record(TEXT("action"));
 
 		if (Action.Len() > 0 && DepotFileName.Len() > 0)
 		{
@@ -1558,8 +1546,8 @@ static void ParseShelvedChangelistResults(const FP4RecordSet& InRecords, TMap<FS
 
 	for (int32 FileIndex = 0;; ++FileIndex)
 	{
-		FString DepotFileName = Record(FString::Printf(TEXT("depotFile%d"), FileIndex));
-		FString Action = Record(FString::Printf(TEXT("action%d"), FileIndex));
+		const FString& DepotFileName = Record(FString::Printf(TEXT("depotFile%d"), FileIndex));
+		const FString& Action = Record(FString::Printf(TEXT("action%d"), FileIndex));
 
 		if (DepotFileName.Len() == 0)
 			break;
@@ -1600,14 +1588,12 @@ static void ParseHistoryResults(FPerforceSourceControlProvider& SCCProvider, con
 	if (InRecords.Num() > 0)
 	{
 		// Iterate over each record, extracting the relevant information for each
-		for (int32 RecordIndex = 0; RecordIndex < InRecords.Num(); ++RecordIndex)
+		for (const FP4Record& ClientRecord : InRecords)
 		{
-			const FP4Record& ClientRecord = InRecords[RecordIndex];
-
 			// Extract the file name
 			check(ClientRecord.Contains(TEXT("depotFile")));
-			FString DepotFileName = ClientRecord(TEXT("depotFile"));
-			FString LocalFileName = FindWorkspaceFile(InStates, DepotFileName);
+			const FString& DepotFileName = ClientRecord(TEXT("depotFile"));
+			const FString& LocalFileName = FindWorkspaceFile(InStates, DepotFileName);
 
 			if (OutHistory.Contains(LocalFileName))
 			{
@@ -1626,32 +1612,32 @@ static void ParseHistoryResults(FPerforceSourceControlProvider& SCCProvider, con
 					// No more revisions
 					break;
 				}
-				FString RevisionNumber = ClientRecord(*VarName);
+				const FString& RevisionNumber = ClientRecord(*VarName);
 
 				// Extract the user name
 				VarName = FString::Printf(TEXT("user%d"), RevisionNumbers);
 				check(ClientRecord.Contains(*VarName));
-				FString UserName = ClientRecord(*VarName);
+				const FString& UserName = ClientRecord(*VarName);
 
 				// Extract the date
 				VarName = FString::Printf(TEXT("time%d"), RevisionNumbers);
 				check(ClientRecord.Contains(*VarName));
-				FString Date = ClientRecord(*VarName);
+				const FString& Date = ClientRecord(*VarName);
 
 				// Extract the changelist number
 				VarName = FString::Printf(TEXT("change%d"), RevisionNumbers);
 				check(ClientRecord.Contains(*VarName));
-				FString ChangelistNumber = ClientRecord(*VarName);
+				const FString& ChangelistNumber = ClientRecord(*VarName);
 
 				// Extract the description
 				VarName = FString::Printf(TEXT("desc%d"), RevisionNumbers);
 				check(ClientRecord.Contains(*VarName));
-				FString Description = ClientRecord(*VarName);
+				const FString& Description = ClientRecord(*VarName);
 
 				// Extract the action
 				VarName = FString::Printf(TEXT("action%d"), RevisionNumbers);
 				check(ClientRecord.Contains(*VarName));
-				FString Action = ClientRecord(*VarName);
+				const FString& Action = ClientRecord(*VarName);
 
 				FString FileSize(TEXT("0"));
 
@@ -1666,7 +1652,7 @@ static void ParseHistoryResults(FPerforceSourceControlProvider& SCCProvider, con
 				// Extract the clientspec/workspace
 				VarName = FString::Printf(TEXT("client%d"), RevisionNumbers);
 				check(ClientRecord.Contains(*VarName));
-				FString ClientSpec = ClientRecord(*VarName);
+				const FString& ClientSpec = ClientRecord(*VarName);
 
 				// check for branch
 				TSharedPtr<FPerforceSourceControlRevision, ESPMode::ThreadSafe> BranchSource;
@@ -1676,11 +1662,11 @@ static void ParseHistoryResults(FPerforceSourceControlProvider& SCCProvider, con
 					BranchSource = MakeShareable(new FPerforceSourceControlRevision(SCCProvider));
 
 					VarName = FString::Printf(TEXT("file%d,0"), RevisionNumbers);
-					FString BranchSourceFileName = ClientRecord(*VarName);
+					const FString& BranchSourceFileName = ClientRecord(*VarName);
 					BranchSource->FileName = FindWorkspaceFile(InStates, BranchSourceFileName);
 
 					VarName = FString::Printf(TEXT("erev%d,0"), RevisionNumbers);
-					FString BranchSourceRevision = ClientRecord(*VarName);
+					const FString& BranchSourceRevision = ClientRecord(*VarName);
 					BranchSource->RevisionNumber = FCString::Atoi(*BranchSourceRevision);
 				}
 
@@ -1742,9 +1728,8 @@ static void ParseDiffResults(const FP4RecordSet& InRecords, TArray<FString>& Out
 	if (InRecords.Num() > 0)
 	{
 		// Iterate over each record found as a result of the command, parsing it for relevant information
-		for (int32 Index = 0; Index < InRecords.Num(); ++Index)
+		for (const FP4Record& ClientRecord : InRecords)
 		{
-			const FP4Record& ClientRecord = InRecords[Index];
 			FString FileName = ClientRecord(TEXT("clientFile"));
 			FPaths::NormalizeFilename(FileName);
 			OutModifiedFiles.Add(FileName);
@@ -1755,11 +1740,9 @@ static void ParseDiffResults(const FP4RecordSet& InRecords, TArray<FString>& Out
 static void ParseChangelistsResults(const FP4RecordSet& InRecords, TArray<FPerforceSourceControlChangelistState>& OutStates)
 {
 	// Iterate over each record found as a result of the command, parsing it for relevant information
-	for (int32 Index = 0; Index < InRecords.Num(); ++Index)
+	for (const FP4Record& ClientRecord : InRecords)
 	{
-		const FP4Record& ClientRecord = InRecords[Index];
-
-		FString ChangelistString = ClientRecord(TEXT("change"));
+		const FString& ChangelistString = ClientRecord(TEXT("change"));
 		int32 ChangelistNumber = FCString::Atoi(*ChangelistString);
 
 		FPerforceSourceControlChangelist Changelist(ChangelistNumber);
@@ -1811,7 +1794,7 @@ bool FPerforceUpdateStatusWorker::Execute(FPerforceSourceControlCommand& InComma
 				}
 				else
 				{
-					for (auto& Branch : StatusBranches )
+					for (const FString& Branch : StatusBranches)
 					{
 						// Check the status branch for updates
 						FString BranchFile;
@@ -2026,11 +2009,9 @@ static bool GetOpenedFilesInChangelist(FPerforceConnection& Connection, FPerforc
 
 static void ParseWhereResults(FP4RecordSet& InRecords, TMap<FString, FString>& DepotToFileMap)
 {
-	for (int32 Index = 0; Index < InRecords.Num(); ++Index)
+	for (const FP4Record& Record : InRecords)
 	{
-		const FP4Record& Record = InRecords[Index];
-
-		FString DepotFile = Record(TEXT("depotFile"));
+		const FString& DepotFile = Record(TEXT("depotFile"));
 		FString ClientFile = Record(TEXT("path")).Replace(TEXT("\\"), TEXT("/"));
 
 		if (DepotFile.Len() > 0 && ClientFile.Len() > 0)
diff --git a/Engine/Plugins/Developer/PerforceSourceControl/Source/PerforceSourceControl/Private/PerforceSourceControlProvider.cpp b/Engine/Plugins/Developer/PerforceSourceControl/Source/PerforceSourceControl/Private/PerforceSourceControlProvider.cpp
--- a/Engine/Plugins/Developer/PerforceSourceControl/Source/PerforceSourceControl/Private/PerforceSourceControlProvider.cpp
+++ b/Engine/Plugins/Developer/PerforceSourceControl/Source/PerforceSourceControl/Private/PerforceSourceControlProvider.cpp
@@ -583,13 +583,12 @@ void FPerforceSourceControlProvider::Tick()
 static void ParseGetLabelsResults(FPerforceSourceControlProvider& InSourceControlProvider, const FP4RecordSet& InRecords, TArray< TSharedRef<ISourceControlLabel> >& OutLabels)
 {
 	// Iterate over each record found as a result of the command, parsing it for relevant information
-	for (int32 Index = 0; Index < InRecords.Num(); ++Index)
+	for (const FP4Record& ClientRecord : InRecords)
 	{
-		const FP4Record& ClientRecord = InRecords[Index];
-		FString LabelName = ClientRecord(TEXT("label"));
+		const FString& LabelName = ClientRecord(TEXT("label"));
 		if(LabelName.Len() > 0)
 		{
-			OutLabels.Add(MakeShareable( new FPerforceSourceControlLabel(InSourceControlProvider, LabelName) ) );
+			OutLabels.Add(MakeShared<FPerforceSourceControlLabel>(InSourceControlProvider, LabelName));
 		}
 	}
 }
diff --git a/Engine/Source/Developer/SourceControl/Private/ISourceControlProvider.cpp b/Engine/Source/Developer/SourceControl/Private/ISourceControlProvider.cpp
--- a/Engine/Source/Developer/SourceControl/Private/ISourceControlProvider.cpp
+++ b/Engine/Source/Developer/SourceControl/Private/ISourceControlProvider.cpp
@@ -72,7 +72,10 @@ ECommandResult::Type ISourceControlProvider::Execute(const FSourceControlOperati
 ECommandResult::Type ISourceControlProvider::Execute(const FSourceControlOperationRef& InOperation, const FString& InFile, const EConcurrency::Type InConcurrency, const FSourceControlOperationComplete& InOperationCompleteDelegate)
 {
 	TArray<FString> FileArray;
-	FileArray.Add(InFile);
+	if (!InFile.IsEmpty())
+	{
+		FileArray.Add(InFile);
+	}
 	return Execute(InOperation, FileArray, InConcurrency, InOperationCompleteDelegate);
 }
 
diff --git a/Engine/Source/Developer/SourceControl/Private/SourceControlModule.cpp b/Engine/Source/Developer/SourceControl/Private/SourceControlModule.cpp
--- a/Engine/Source/Developer/SourceControl/Private/SourceControlModule.cpp
+++ b/Engine/Source/Developer/SourceControl/Private/SourceControlModule.cpp
@@ -238,8 +238,8 @@ void FSourceControlModule::GetProviderNames(TArray<FName>& OutProviderNames)
 }
 
 void FSourceControlModule::Tick()
-{	
-	if( CurrentSourceControlProvider != nullptr )
+{
+	if (CurrentSourceControlProvider != nullptr)
 	{
 		ISourceControlProvider& Provider = GetProvider();
 
@@ -249,23 +249,24 @@ void FSourceControlModule::Tick()
 		AssetDataCache.Tick();
 
 		// don't allow background status updates when temporarily disabled for login
-		if(!bTemporarilyDisabled)
+		if (!bTemporarilyDisabled)
 		{
 			// check for any pending dispatches
-			if(PendingStatusUpdateFiles.Num() > 0)
+			if (PendingStatusUpdateFiles.Num() > 0)
 			{
 				// grab a batch of files
 				TArray<FString> FilesToDispatch;
-				for(auto Iter(PendingStatusUpdateFiles.CreateConstIterator()); Iter; Iter++)
+				FilesToDispatch.Reserve(SourceControlConstants::MaxStatusDispatchesPerTick);
+				for (const FString& Filename : PendingStatusUpdateFiles)
 				{
 					if(FilesToDispatch.Num() >= SourceControlConstants::MaxStatusDispatchesPerTick)
 					{
 						break;
 					}
-					FilesToDispatch.Add(*Iter);
+					FilesToDispatch.Add(Filename);
 				}
 
-				if(FilesToDispatch.Num() > 0)
+				if (FilesToDispatch.Num() > 0)
 				{
 					// remove the files we are dispatching so we don't try again
 					PendingStatusUpdateFiles.RemoveAt(0, FilesToDispatch.Num());
@@ -280,29 +281,29 @@ void FSourceControlModule::Tick()
 
 void FSourceControlModule::QueueStatusUpdate(const TArray<UPackage*>& InPackages)
 {
-	if(IsEnabled())
+	if (IsEnabled())
 	{
-		for(auto It(InPackages.CreateConstIterator()); It; It++)
+		for (UPackage* Package: InPackages)
 		{
-			QueueStatusUpdate(*It);
+			QueueStatusUpdate(Package);
 		}
 	}
 }
 
 void FSourceControlModule::QueueStatusUpdate(const TArray<FString>& InFilenames)
 {
-	if(IsEnabled())
+	if (IsEnabled())
 	{
-		for(auto It(InFilenames.CreateConstIterator()); It; It++)
+		for (const FString& Filename : InFilenames)
 		{
-			QueueStatusUpdate(*It);
+			QueueStatusUpdate(Filename);
 		}
 	}
 }
 
 void FSourceControlModule::QueueStatusUpdate(UPackage* InPackage)
 {
-	if(IsEnabled())
+	if (IsEnabled())
 	{
 		QueueStatusUpdate(SourceControlHelpers::PackageFilename(InPackage));
 	}
@@ -310,13 +311,13 @@ void FSourceControlModule::QueueStatusUpdate(UPackage* InPackage)
 
 void FSourceControlModule::QueueStatusUpdate(const FString& InFilename)
 {
-	if(IsEnabled())
+	if (IsEnabled())
 	{
 		TSharedPtr<ISourceControlState, ESPMode::ThreadSafe> SourceControlState = GetProvider().GetState(InFilename, EStateCacheUsage::Use);
-		if(SourceControlState.IsValid())
+		if (SourceControlState.IsValid())
 		{
 			FTimespan TimeSinceLastUpdate = FDateTime::Now() - SourceControlState->GetTimeStamp();
-			if(TimeSinceLastUpdate > SourceControlConstants::StateRefreshInterval)
+			if (TimeSinceLastUpdate > SourceControlConstants::StateRefreshInterval)
 			{
 				PendingStatusUpdateFiles.AddUnique(InFilename);
 			}
