Overriding the perforce server address (P4PORT) no longer requires that the user (P4USER) and workspace (P4CLIENT) be overridden as well.

#rnx

- Changed check so that we only check if  'NewServerName' is set in order to apply it as the P4PORT.
-- Originally P4PORT, P4USER and P4CLIENT were all set inside of this if statement and it seems to have been assumed that either all three of the input parameters will be set, or none of them.
-- Host masquerading was added later and was specifically done so that it would only be applied if the P4PORT is overridden as well. I am not sure on the use case for this so I opted to leave it alone.
- Some of the earlier connection attempts don't use the user name or the workspace name so reporting them as part of the failure can be confusing. I have removed them.

diff --git a/Engine/Plugins/Developer/PerforceSourceControl/Source/PerforceSourceControl/Private/PerforceConnection.cpp b/Engine/Plugins/Developer/PerforceSourceControl/Source/PerforceSourceControl/Private/PerforceConnection.cpp
--- a/Engine/Plugins/Developer/PerforceSourceControl/Source/PerforceSourceControl/Private/PerforceConnection.cpp
+++ b/Engine/Plugins/Developer/PerforceSourceControl/Source/PerforceSourceControl/Private/PerforceConnection.cpp
@@ -535,12 +535,11 @@ bool FPerforceConnection::EnsureValidConnection(FString& InOutServerName, FStrin
 	TestP4.SetProtocol("tag", "");
 	TestP4.SetProtocol("enableStreams", "");
 
-	if (NewServerName.Len() && NewUserName.Len() && NewClientSpecName.Len())
+	if (!NewServerName.IsEmpty())
 	{
-		//attempt connection with given settings
 		TestP4.SetPort(TCHAR_TO_ANSI(*NewServerName));
 
-		if(InConnectionInfo.HostOverride.Len() > 0)
+		if (!InConnectionInfo.HostOverride.IsEmpty())
 		{
 			TestP4.SetHost(TCHAR_TO_ANSI(*InConnectionInfo.HostOverride));
 		}
@@ -566,10 +565,8 @@ bool FPerforceConnection::EnsureValidConnection(FString& InOutServerName, FStrin
 		FFormatNamedArguments Arguments;
 		Arguments.Add(TEXT("OwningSystem"), FText::FromString(SCCProvider.GetOwnerName()));
 		Arguments.Add( TEXT("PortName"), FText::FromString(NewServerName) );
-		Arguments.Add( TEXT("UserName"), FText::FromString(NewUserName) );
-		Arguments.Add( TEXT("ClientSpecName"), FText::FromString(NewClientSpecName) );
 		Arguments.Add( TEXT("Ticket"), FText::FromString(InConnectionInfo.Ticket) );
-		SourceControlLog.Error(FText::Format(LOCTEXT("P4ConnectErrorConnection_Details", "OwningSystem={OwningSystem}, Port={PortName}, User={UserName}, ClientSpec={ClientSpecName}, Ticket={Ticket}"), Arguments));
+		SourceControlLog.Error(FText::Format(LOCTEXT("P4ConnectErrorConnection_Details", "OwningSystem={OwningSystem}, Port={PortName}, Ticket={Ticket}"), Arguments));
 	}
 
 	// run an info command to determine unicode status
@@ -585,10 +582,8 @@ bool FPerforceConnection::EnsureValidConnection(FString& InOutServerName, FStrin
 			FFormatNamedArguments Arguments;
 			Arguments.Add(TEXT("OwningSystem"), FText::FromString(SCCProvider.GetOwnerName()));
 			Arguments.Add( TEXT("PortName"), FText::FromString(NewServerName) );
-			Arguments.Add( TEXT("UserName"), FText::FromString(NewUserName) );
-			Arguments.Add( TEXT("ClientSpecName"), FText::FromString(NewClientSpecName) );
 			Arguments.Add( TEXT("Ticket"), FText::FromString(InConnectionInfo.Ticket) );
-			SourceControlLog.Error(FText::Format(LOCTEXT("P4UnicodeErrorConnection_Details", "OwningSystem={OwningSystem}, Port={PortName}, User={UserName}, ClientSpec={ClientSpecName}, Ticket={Ticket}"), Arguments));
+			SourceControlLog.Error(FText::Format(LOCTEXT("P4UnicodeErrorConnection_Details", "OwningSystem={OwningSystem}, Port={PortName}, Ticket={Ticket}"), Arguments));
 		}
 		else
 		{
