The perforce source control API can now issue a single p4 print command with multiple depot file paths as arguments rather than issuing a single p4 print command per file.

#rnx

- Clean up the source control backend by removing the older code path that was ifdefed out and clean up the new path.
- FPerforceConnection no longer takes a TOptional<FSharedBuffer> but takes an optional pointer to a TArray of FSharedBuffer. When gathering files this should return one FSharedBuffer per file requested.
-- It might make more sense to attach the FSharedBuffer to the FP4Record but that would require some bigger changes and would be better suited to being some follow on work. All the code is private in the module and not exposed so it is fairly safe and easy to do these refactors later.

- Replaced the TArray buffer for gathering files with a purpose built class FP4DepotFile.
- When issuing a p4 print command with multiple files we will get one call to ClientUser::OutputStat followed by calls to either ClientUser::OutputText or ClientUser::OutputBinary until all of the file has been delivered, then we will get another call to ClientUser::OutputStat for the next file and so on.
-- Note with binary files we seem to get one final call to ClientUser::OutputBinary with zero size, which I assume is supposed to signal the end of the transfer, text files do not get this, they seem to just stop.

diff --git a/Engine/Plugins/Developer/PerforceSourceControl/Source/PerforceSourceControl/Private/PerforceConnection.cpp b/Engine/Plugins/Developer/PerforceSourceControl/Source/PerforceSourceControl/Private/PerforceConnection.cpp
--- a/Engine/Plugins/Developer/PerforceSourceControl/Source/PerforceSourceControl/Private/PerforceConnection.cpp
+++ b/Engine/Plugins/Developer/PerforceSourceControl/Source/PerforceSourceControl/Private/PerforceConnection.cpp
@@ -19,27 +19,144 @@
 #define FROM_TCHAR(InText, bIsUnicodeServer) (bIsUnicodeServer ? TCHAR_TO_UTF8(InText) : TCHAR_TO_ANSI(InText))
 #define TO_TCHAR(InText, bIsUnicodeServer) (bIsUnicodeServer ? UTF8_TO_TCHAR(InText) : ANSI_TO_TCHAR(InText))
 
+const FString FP4Record::EmptyStr;
+
 namespace
 {
-	enum class EP4ClientUserFlags
+
+enum class EP4ClientUserFlags
+{
+	None = 0,
+	/** The server uses unicode */
+	UnicodeServer = 1 << 0,
+	/** Binary data returned by commands should be collected in the DataBuffer member */
+	CollectData = 1 << 1,
+};
+
+ENUM_CLASS_FLAGS(EP4ClientUserFlags);
+
+/**
+ * A utility class to make it easier to gather a depot file from perforce when running
+ * p4 print.
+ */
+class FP4DepotFile
+{
+public:
+	FP4DepotFile() = default;
+	~FP4DepotFile() = default;
+
+	/** 
+	 * Start gathering the file in the given record. If the record is missing data
+	 * then the gather will not begin. The calling code should check for this and
+	 * raise errors or warnings accordingly.
+	 * This class does not actually do any perforce work itself, and relies on a
+	 * ClientUser to actually provide the data as it is downloaded.
+	 */
+	void Initialize(const FP4Record& Record)
 	{
-		None			= 0,
-		/** The server uses unicode */
-		UnicodeServer	= 1 << 0,
-		/** Binary data returned by commands should be collected in the DataBuffer member */
-		CollectData		= 1 << 1,
-	};
+		const FString& SizeAsString = Record(TEXT("fileSize"));
+		DepotFilePath = Record(TEXT("depotFile"));
 
-	ENUM_CLASS_FLAGS(EP4ClientUserFlags);
-}
+		if (!SizeAsString.IsEmpty() && !DepotFilePath.IsEmpty())
+		{
+			int64 FileSize = FCString::Atoi64(*SizeAsString);
 
-const FString FP4Record::EmptyStr;
+			Data = FUniqueBuffer::Alloc(FileSize);
+			Offset = 0;
+		}
+	}
+
+	/** Returns true if the FP4DepotFile was set up correctly and can gather a file. */
+	bool IsValid() const
+	{
+		return !Data.IsNull();
+	}
 
+	/** Returns true if all of the files data has been acquired. */
+	bool IsFileComplete() const
+	{
+		return Offset == (int64)Data.GetSize();
+	}
+
+	/** Returns the number of bytes in the file that have not yet been acquired. */
+	int64 GetRemainingBytes() const
+	{
+		return (int64)Data.GetSize() - Offset;
+	}
+
+	/** Returns the depot path of the file we are gathering */
+	const FString& GetDepotPath() const
+	{
+		return DepotFilePath;
+	}
+
+	/** 
+	 * Returns the currently acquired file data and then invalidates the FP4DepotFile.
+	 * It is up to the caller to ensure that the entire file has been acquired or to
+	 * decide if a partially acquired file is okay.
+	 */
+	FSharedBuffer Release()
+	{
+		Offset = INDEX_NONE;
+
+		DepotFilePath.Reset();
+
+		return Data.MoveToShared();
+	}
+
+	/* Used to reset the FP4DepotFile if an error is encountered */
+	void Reset()
+	{
+		Offset = INDEX_NONE;
+
+		DepotFilePath.Reset();
+
+		Data.Reset();
+	}
+
+	/** 
+	 * Called when new data for the file has been downloaded and we can add it to the
+	 * data that we have already required.
+	 * 
+	 * @param DataPtr		The pointer to the downloaded data
+	 * @param DataLength	The size of the downloaded data in bytes
+	 * @return				True if the FP4DepotFile is valid and there was enough space
+	 *						for the downloaded data. False if the FP4DepotFile is invalid
+	 *						or if there is not enough space.		
+	 */
+	bool OnDataDownloaded(const char* DataPtr, int DataLength)
+	{
+		if (DataLength <= GetRemainingBytes())
+		{
+			FMemory::Memcpy((char*)Data.GetData() + Offset, DataPtr, DataLength);
+			Offset += DataLength;
+
+			return true;
+		}
+		else
+		{
+			return false;
+		}
+	}
+
+private:
+
+	/** The path of the file in the perforce depot */
+	FString DepotFilePath;
+
+	/** The buffer containing the file data, allocated up front. */
+	FUniqueBuffer Data;
+	/** Tracks where the next set on downloaded data should be placed in the buffer. */
+	int64 Offset = INDEX_NONE;
+};
+
+} //namespace
 
 /** Custom ClientUser class for handling results and errors from Perforce commands */
 class FP4ClientUser : public ClientUser
 {
 public:
+
 	FP4ClientUser(FP4RecordSet& InRecords, EP4ClientUserFlags InFlags, TArray<FText>& InOutErrorMessages)
 		: ClientUser()
 		, Flags(InFlags)
@@ -61,19 +178,17 @@ public:
 			Record.Add(TO_TCHAR(Var.Text(), IsUnicodeServer()), TO_TCHAR(Value.Text(), IsUnicodeServer()));
 		}
 
-		// Try to reserve the data array if we have one as we just got the filesize info.
-		// NOTE: This code is built on the assumption that the data array will contain the
-		// data sent via p4 print and our code only called p4 print on a single file at the 
-		// time.
-		if (IsCollectingData() && DataBuffer.GetAllocatedSize() == 0)
+		if (IsCollectingData())
 		{
-			const FString& SizeAsString = Record(TEXT("fileSize"));
-			if (!SizeAsString.IsEmpty())
+			if (File.IsValid())
 			{
-				const int64 Size = FCString::Atoi64(*SizeAsString);
-				DataBuffer.Reserve(Size);
+				FText Message = FText::Format(LOCTEXT("P4Client_GatheringUnfinished", "Started gathering depot file '{0}' before the previous file finished!"),
+					FText::FromString(File.GetDepotPath()));
+
+				OutErrorMessages.Add(MoveTemp(Message));
 			}
 
+			File.Initialize(Record);
 		}
 
 		Records.Add(Record);
@@ -84,7 +199,25 @@ public:
 	{
 		if (IsCollectingData())
 		{
-			DataBuffer.Append(DataPtr, DataLength);
+			if (File.OnDataDownloaded(DataPtr, DataLength))
+			{
+				if (File.IsFileComplete())
+				{
+					Files.Add(File.Release());
+				}
+			}
+			else
+			{
+				FText Message = FText::Format(	
+					LOCTEXT("P4Client_TextCollectionFailed", "Collecting text data requires {0} bytes but the buffer only has {1} bytes remaining: {2}"),
+					DataLength,
+					File.GetRemainingBytes(),
+					FText::FromString(File.GetDepotPath()));
+
+				OutErrorMessages.Add(MoveTemp(Message));
+
+				File.Reset();
+			}
 		}
 		else
 		{
@@ -97,7 +230,38 @@ public:
 	{
 		if (IsCollectingData())
 		{
-			DataBuffer.Append(DataPtr, DataLength);
+			// For binary files we get a zero size call once the file is completed so we wait for that
+			// rather than checking FP4DepotFile::isFileComplete after every transfer
+			if (DataLength == 0)
+			{
+				if (File.IsFileComplete())
+				{
+					Files.Add(File.Release());
+				}
+				else
+				{
+					FText Message = FText::Format(
+						LOCTEXT("P4Client_IncompleteFIle", "Collecting binary data completed but missing {0} bytes: {1}"),
+						File.GetRemainingBytes(),
+						FText::FromString(File.GetDepotPath()));
+
+					OutErrorMessages.Add(MoveTemp(Message));
+
+					File.Reset();
+				}
+			}
+			else if (!File.OnDataDownloaded(DataPtr, DataLength))
+			{
+				FText Message = FText::Format(
+					LOCTEXT("P4Client_BinaryCollectionFailed", "Collecting binary data requires {0} bytes but the buffer only has {1} bytes remaining: {2}"),
+					DataLength,
+					File.GetRemainingBytes(),
+					FText::FromString(File.GetDepotPath()));
+
+				OutErrorMessages.Add(MoveTemp(Message));
+
+				File.Reset();
+			}
 		}
 		else
 		{
@@ -105,7 +269,7 @@ public:
 		}
 	}
 
-	/** Called by P4API when an error message is avaliable. */
+	/** Called by P4API when an error message is available. */
 	virtual void HandleError(Error* InError) override
 	{
 		StrBuf ErrorMessage;
@@ -124,9 +288,9 @@ public:
 	}
 
 	/** Returns DataBuffer as a FSharedBuffer, note that once called DataBuffer will be empty */
-	inline FSharedBuffer ReleaseData()
+	inline TArray<FSharedBuffer> ReleaseData()
 	{
-		return MakeSharedBufferFromArray(MoveTemp(DataBuffer));
+		return MoveTemp(Files);
 	}
 
 	EP4ClientUserFlags Flags;
@@ -134,7 +298,9 @@ public:
 	TArray<FText>& OutErrorMessages;
 
 private:
-	TArray64<char> DataBuffer;
+	TArray<FSharedBuffer> Files;
+
+	FP4DepotFile File;
 };
 
 /** A class used instead of FP4ClientUser for handling changelist create command */
@@ -814,7 +980,7 @@ void FPerforceConnection::Disconnect()
 }
 
 bool FPerforceConnection::RunCommand(	const FString& InCommand, const TArray<FString>& InParameters, FP4RecordSet& OutRecordSet, 
-										TOptional<FSharedBuffer>& OutData, TArray<FText>& OutErrorMessage, 
+										TArray<FSharedBuffer>* OutData, TArray<FText>& OutErrorMessage, 
 										FOnIsCancelled InIsCancelled, bool& OutConnectionDropped, ERunCommandFlags RunFlags)
 {
 	TRACE_CPUPROFILER_EVENT_SCOPE_TEXT(*FString::Printf(TEXT("FPerforceConnection::RunCommand_%s"), *InCommand));
@@ -868,7 +1034,7 @@ bool FPerforceConnection::RunCommand(	const FString& InCommand, const TArray<FSt
 
 	EP4ClientUserFlags ClientUserFlags = EP4ClientUserFlags::None;
 	ClientUserFlags |= bIsUnicode ? EP4ClientUserFlags::UnicodeServer : EP4ClientUserFlags::None;
-	ClientUserFlags |= OutData ? EP4ClientUserFlags::CollectData : EP4ClientUserFlags::None;
+	ClientUserFlags |= OutData != nullptr ? EP4ClientUserFlags::CollectData : EP4ClientUserFlags::None;
 	
 	FP4ClientUser User(OutRecordSet, ClientUserFlags, OutErrorMessage);
 	if (EnumHasAllFlags(RunFlags, ERunCommandFlags::Quiet))
@@ -882,9 +1048,9 @@ bool FPerforceConnection::RunCommand(	const FString& InCommand, const TArray<FSt
 		OutConnectionDropped = true;
 	}
 
-	if (OutData)
+	if (OutData != nullptr)
 	{
-		OutData = User.ReleaseData();
+		*OutData = User.ReleaseData();
 	}	
 
 	P4Client.SetBreak(nullptr);
@@ -938,7 +1104,6 @@ int32 FPerforceConnection::EditPendingChangelist(const FText& NewDescription, in
 	// Get changelist current specification
 	{
 		TArray<FString> Params;
-		TOptional<FSharedBuffer> CommandData;
 		bool bConnectionDropped = false;
 		const ERunCommandFlags RunFlags = ERunCommandFlags::DisableCommandLogging;
 
@@ -946,7 +1111,7 @@ int32 FPerforceConnection::EditPendingChangelist(const FText& NewDescription, in
 		// TODO : make this work also for default changelist, but should really be a Create
 		Params.Add(FString::Printf(TEXT("%d"), ChangelistNumber));
 
-		if (!RunCommand(TEXT("change"), Params, PreviousRecords, CommandData, OutErrorMessages, InIsCancelled, bConnectionDropped, RunFlags))
+		if (!RunCommand(TEXT("change"), Params, PreviousRecords, nullptr, OutErrorMessages, InIsCancelled, bConnectionDropped, RunFlags))
 		{
 			return 0;
 		}
@@ -1049,13 +1214,12 @@ void FPerforceConnection::EstablishConnection(const FPerforceConnectionInfo& InC
 		TArray<FString> Params;
 		TArray<FText> ErrorMessages;
 		FP4RecordSet Records;
-		TOptional<FSharedBuffer> CommandData;
 		bool bConnectionDropped = false;
 		const ERunCommandFlags RunFlags = ERunCommandFlags::DisableCommandLogging;
 
 		UE_LOG(LogSourceControl, Verbose, TEXT(" ... checking unicode status" ));
 
-		if (RunCommand(TEXT("info"), Params, Records, CommandData, ErrorMessages, FOnIsCancelled(), bConnectionDropped, RunFlags))
+		if (RunCommand(TEXT("info"), Params, Records, nullptr, ErrorMessages, FOnIsCancelled(), bConnectionDropped, RunFlags))
 		{
 			// Get character encoding
 			bIsUnicode = Records[0].Find(TEXT("unicode")) != nullptr;
@@ -1088,7 +1252,7 @@ void FPerforceConnection::EstablishConnection(const FPerforceConnectionInfo& InC
 			// Gather the client root
 			UE_LOG(LogSourceControl, Verbose, TEXT(" ... getting info" ));
 			bConnectionDropped = false;
-			if (RunCommand(TEXT("info"), Params, Records, CommandData, ErrorMessages, FOnIsCancelled(), bConnectionDropped, RunFlags))
+			if (RunCommand(TEXT("info"), Params, Records, nullptr, ErrorMessages, FOnIsCancelled(), bConnectionDropped, RunFlags))
 			{
 				UE_LOG(LogSourceControl, Verbose, TEXT(" ... getting clientroot" ));
 				ClientRoot = Records[0](TEXT("clientRoot"));
diff --git a/Engine/Plugins/Developer/PerforceSourceControl/Source/PerforceSourceControl/Private/PerforceConnection.h b/Engine/Plugins/Developer/PerforceSourceControl/Source/PerforceSourceControl/Private/PerforceConnection.h
--- a/Engine/Plugins/Developer/PerforceSourceControl/Source/PerforceSourceControl/Private/PerforceConnection.h
+++ b/Engine/Plugins/Developer/PerforceSourceControl/Source/PerforceSourceControl/Private/PerforceConnection.h
@@ -107,14 +107,13 @@ public:
 	 */
 	bool RunCommand(const FString& InCommand, const TArray<FString>& InParameters, FP4RecordSet& OutRecordSet, TArray<FText>&  OutErrorMessage, FOnIsCancelled InIsCancelled, bool& OutConnectionDropped)
 	{
-		TOptional<FSharedBuffer> UnsetDataBuffer;
-		return RunCommand(InCommand, InParameters, OutRecordSet, UnsetDataBuffer, OutErrorMessage, InIsCancelled, OutConnectionDropped);
+		return RunCommand(InCommand, InParameters, OutRecordSet, nullptr, OutErrorMessage, InIsCancelled, OutConnectionDropped);
 	}
 
 	/**
 	 * Runs internal perforce command, catches exceptions, returns results
 	 */
-	bool RunCommand(const FString& InCommand, const TArray<FString>& InParameters, FP4RecordSet& OutRecordSet, TOptional<FSharedBuffer>& OutData, TArray<FText>& OutErrorMessage, FOnIsCancelled InIsCancelled, bool& OutConnectionDropped, ERunCommandFlags RunFlags = ERunCommandFlags::Default);
+	bool RunCommand(const FString& InCommand, const TArray<FString>& InParameters, FP4RecordSet& OutRecordSet, TArray<FSharedBuffer>* OutData, TArray<FText>& OutErrorMessage, FOnIsCancelled InIsCancelled, bool& OutConnectionDropped, ERunCommandFlags RunFlags = ERunCommandFlags::Default);
 
 	/**
 	 * Creates a changelist with the specified description
diff --git a/Engine/Plugins/Developer/PerforceSourceControl/Source/PerforceSourceControl/Private/PerforceSourceControlOperations.cpp b/Engine/Plugins/Developer/PerforceSourceControl/Source/PerforceSourceControl/Private/PerforceSourceControlOperations.cpp
--- a/Engine/Plugins/Developer/PerforceSourceControl/Source/PerforceSourceControl/Private/PerforceSourceControlOperations.cpp
+++ b/Engine/Plugins/Developer/PerforceSourceControl/Source/PerforceSourceControl/Private/PerforceSourceControlOperations.cpp
@@ -3049,26 +3049,29 @@ bool FPerforceDownloadFileWorker::Execute(FPerforceSourceControlCommand& InComma
 		if (TargetDirectory.IsEmpty())
 		{
 			// Download and store the files as blobs in memory that the caller can access as they wish
-			for (const FString& TargetFilePath : InCommand.Files)
+			Parameters.Add(TEXT("-q")); // Do not print the header, we only want the actual file contents
+			Parameters.Append(InCommand.Files);
+
+			TArray<FSharedBuffer> FilesData;
+			const ERunCommandFlags Flags = Operation->ShouldLogToStdOutput() ? ERunCommandFlags::Default : ERunCommandFlags::DisableCommandLogging;
+
+			InCommand.bCommandSuccessful = Connection.RunCommand(	TEXT("print"), Parameters, Records, &FilesData, InCommand.ResultInfo.ErrorMessages,
+																	FOnIsCancelled::CreateRaw(&InCommand, &FPerforceSourceControlCommand::IsCanceled), 
+																	InCommand.bConnectionDropped, Flags);
+			
+			if (!InCommand.ResultInfo.ErrorMessages.IsEmpty() || (InCommand.Files.Num() != FilesData.Num()))
 			{
-				Parameters.Add(TEXT("-q")); // Do not print the header, we only want the actual file contents
-				Parameters.Add(TargetFilePath);
+				InCommand.bCommandSuccessful = false;
+			}
 
-				TOptional<FSharedBuffer> FileData = FSharedBuffer();
-				const ERunCommandFlags Flags = Operation->ShouldLogToStdOutput() ? ERunCommandFlags::Default : ERunCommandFlags::DisableCommandLogging;
+			if (InCommand.bCommandSuccessful)
+			{
+				check(InCommand.Files.Num() == FilesData.Num());
 
-				InCommand.bCommandSuccessful = Connection.RunCommand(	TEXT("print"), Parameters, Records, FileData, InCommand.ResultInfo.ErrorMessages, 
-																		FOnIsCancelled::CreateRaw(&InCommand, &FPerforceSourceControlCommand::IsCanceled), 
-																		InCommand.bConnectionDropped, Flags);
-				
-				if (InCommand.bCommandSuccessful)
+				for (int32 Index = 0; Index < InCommand.Files.Num(); ++Index)
 				{
-					Operation->AddFileData(TargetFilePath, FileData.GetValue());
+					Operation->AddFileData(InCommand.Files[Index], FilesData[Index]);
 				}
-
-				// Keep the already allocated memory for future files to prevent reallocation
-				Parameters.Reset();
-				Records.Reset();
 			}
 		}
 		else
@@ -3087,10 +3090,9 @@ bool FPerforceDownloadFileWorker::Execute(FPerforceSourceControlCommand& InComma
 				Parameters.Add(TEXT("-q")); // Do not print the header, we only want the actual file contents
 				Parameters.Add(TargetFilePath);
 
-				TOptional<FSharedBuffer> NullBuffer;
 				const ERunCommandFlags Flags = Operation->ShouldLogToStdOutput() ? ERunCommandFlags::Default : ERunCommandFlags::DisableCommandLogging;
 
-				InCommand.bCommandSuccessful = Connection.RunCommand(	TEXT("print"), Parameters, Records, NullBuffer, InCommand.ResultInfo.ErrorMessages, 
+				InCommand.bCommandSuccessful = Connection.RunCommand(	TEXT("print"), Parameters, Records, nullptr, InCommand.ResultInfo.ErrorMessages, 
 																		FOnIsCancelled::CreateRaw(&InCommand, &FPerforceSourceControlCommand::IsCanceled), 
 																		InCommand.bConnectionDropped, Flags);
 
@@ -3243,7 +3245,6 @@ bool FPerforceDeleteWorkspaceWorker::Execute(FPerforceSourceControlCommand& InCo
 		Parameters.Add(Operation->GetWorkspaceName());
 
 		FP4RecordSet Records;
-		TOptional<FSharedBuffer> NullBuffer;
 		// p4 client command does not return tagged results, which means info will get printed to
 		// stdout which we do not want, so we can use the Quiet flag which is the same as passing
 		// in -q as a global-opt in the command.
@@ -3252,7 +3253,7 @@ bool FPerforceDeleteWorkspaceWorker::Execute(FPerforceSourceControlCommand& InCo
 		Connection.RunCommand(	TEXT("client"), 
 								Parameters, 
 								Records, 
-								NullBuffer,
+								nullptr,
 								InCommand.ResultInfo.ErrorMessages, 
 								FOnIsCancelled::CreateRaw(&InCommand, &FPerforceSourceControlCommand::IsCanceled), 
 								InCommand.bConnectionDropped,
diff --git a/Engine/Source/Developer/Virtualization/Private/VirtualizationSourceControlBackend.cpp b/Engine/Source/Developer/Virtualization/Private/VirtualizationSourceControlBackend.cpp
--- a/Engine/Source/Developer/Virtualization/Private/VirtualizationSourceControlBackend.cpp
+++ b/Engine/Source/Developer/Virtualization/Private/VirtualizationSourceControlBackend.cpp
@@ -204,13 +204,26 @@ static void CreateDescription(const FString& ProjectName, TArrayView<const FPush
 	return SCCProvider.GetState(DepotPaths, OutStates, EStateCacheUsage::Use);
 }
 
-/** Parse all error messages in a FSourceControlResultInfo and return true if the file not found error message is found */
-[[nodiscard]] static bool IsDepotFileMissing(const FSourceControlResultInfo& ResultInfo)
+/** 
+ * If the only reason we failed an operation is because of missing depot files then there
+ * is no point retrying that operation however other kinds of errors, such as connection
+ * problems, might be not be encountered if we try again.
+ * 
+ * @param ResultInfo	The results of a failed operation
+ * @return				True if there were no errors (as we don't really know what went wrong)
+ *						or if there errors not relating to missing files, otherwise false.
+ */
+[[nodiscard]] static bool ShouldRetryOperation(const FSourceControlResultInfo& ResultInfo)
 {
+	if (ResultInfo.ErrorMessages.IsEmpty())
+	{
+		return true;
+	}
+
 	// Ideally we'd parse for this sort of thing in the source control module itself and return an error enum
 	for (const FText& ErrorTest : ResultInfo.ErrorMessages)
 	{
-		if (ErrorTest.ToString().Find(" - no such file(s).") != INDEX_NONE)
+		if (ErrorTest.ToString().Find(" - no such file(s).") == INDEX_NONE)
 		{
 			return true;
 		}
@@ -352,66 +365,6 @@ bool FSourceControlBackend::PullData(TArrayView<FPullRequest> Requests)
 {
 	TRACE_CPUPROFILER_EVENT_SCOPE(FSourceControlBackend::PullData);
 
-#if 0
-	TRACE_CPUPROFILER_EVENT_SCOPE(FSourceControlBackend::PullData);
-
-	for (FPullRequest& Request : Requests)
-	{
-		TStringBuilder<512> DepotPath;
-		CreateDepotPath(Request.Identifier, DepotPath);
-
-		// TODO: When multiple threads are blocked waiting on this we could gather X payloads together and make a single
-		// batch request on the same connection, which should be a lot faster with less overhead.
-		// Although ideally this backend will not get hit very often.
-		FSemaphoreScopeLock _(ConcurrentConnectionLimit.Get());
-
-		UE_LOG(LogVirtualization, Verbose, TEXT("[%s] Attempting to pull '%s' from source control"), *GetDebugName(), *DepotPath);
-
-		int32 Retries = 0;
-
-		while (Retries < RetryCount)
-		{
-			// Only warn if the backend is configured to retry
-			if (Retries != 0)
-			{
-				UE_LOG(LogVirtualization, Warning, TEXT("[%s] Failed to download '%s' retrying (%d/%d) in %dms..."), *GetDebugName(), DepotPath.ToString(), Retries, RetryCount, RetryWaitTimeMS);
-				FPlatformProcess::SleepNoStats(RetryWaitTimeMS * 0.001f);
-			}
-
-#if IS_SOURCE_CONTROL_THREAD_SAFE
-			TSharedRef<FDownloadFile, ESPMode::ThreadSafe> DownloadCommand = ISourceControlOperation::Create<FDownloadFile>(FDownloadFile::EVerbosity::None);
-			if (SCCProvider->Execute(DownloadCommand, DepotPath.ToString(), EConcurrency::Synchronous) == ECommandResult::Succeeded)
-			{
-				// The payload was created by FCompressedBuffer::Compress so we can return it as a FCompressedBuffer.
-				FSharedBuffer Buffer = DownloadCommand->GetFileData(DepotPath);
-				return FCompressedBuffer::FromCompressed(Buffer);
-			}
-#else
-			TSharedRef<FDownloadFile> DownloadCommand = ISourceControlOperation::Create<FDownloadFile>(FDownloadFile::EVerbosity::None);
-			if (SCCProvider->TryToDownloadFileFromBackgroundThread(DownloadCommand, DepotPath.ToString()))
-			{
-				// The payload was created by FCompressedBuffer::Compress so we can return it as a FCompressedBuffer.
-				Request.Status = FPullRequest::EStatus::Success;
-				Request.Payload = FCompressedBuffer::FromCompressed(DownloadCommand->GetFileData(DepotPath));
-
-				break;
-			}
-#endif
-
-			// If this was the first try then check to see if the error being returns is that the file does not exist
-			// in the depot. If it does not exist then there is no point in us retrying and we can error out at this point.
-			if (Retries == 0 && IsDepotFileMissing(DownloadCommand->GetResultInfo()))
-			{
-				break;
-			}
-
-			Retries++;
-		}
-	}
-
-	return true;
-
-#else
 	TArray<FString> DepotPaths;
 	DepotPaths.Reserve(Requests.Num());
 
@@ -428,8 +381,6 @@ bool FSourceControlBackend::PullData(TArrayView<FPullRequest> Requests)
 	// Although ideally this backend will not get hit very often.
 	FSemaphoreScopeLock _(ConcurrentConnectionLimit.Get());
 
-//	UE_LOG(LogVirtualization, Verbose, TEXT("[%s] Attempting to pull '%s' from source control"), *GetDebugName(), *DepotPath);
-
 	int32 Retries = 0;
 
 	while (Retries < RetryCount)
@@ -443,44 +394,37 @@ bool FSourceControlBackend::PullData(TArrayView<FPullRequest> Requests)
 
 #if IS_SOURCE_CONTROL_THREAD_SAFE
 		TSharedRef<FDownloadFile, ESPMode::ThreadSafe> DownloadCommand = ISourceControlOperation::Create<FDownloadFile>(FDownloadFile::EVerbosity::None);
-		if (SCCProvider->Execute(DownloadCommand, DepotPath.ToString(), EConcurrency::Synchronous) == ECommandResult::Succeeded)
-		{
-			// The payload was created by FCompressedBuffer::Compress so we can return it as a FCompressedBuffer.
-			FSharedBuffer Buffer = DownloadCommand->GetFileData(DepotPath);
-
-			Request.Payload = FCompressedBuffer::FromCompressed(Buffer);
-			Request.Status = FPullRequest::EStatus::Success;
-		}
+		const bool bOperationSuccess = SCCProvider->Execute(DownloadCommand, DepotPaths, EConcurrency::Synchronous) == ECommandResult::Succeeded;
 #else
 		TSharedRef<FDownloadFile> DownloadCommand = ISourceControlOperation::Create<FDownloadFile>(FDownloadFile::EVerbosity::None);
-		if (SCCProvider->TryToDownloadFileFromBackgroundThread(DownloadCommand, DepotPaths))
-		{
-			for(int32 Index = 0; Index < Requests.Num(); ++Index)
-			{
-				// The payload was created by FCompressedBuffer::Compress so we can return it as a FCompressedBuffer.
-				FSharedBuffer Buffer = DownloadCommand->GetFileData(DepotPaths[Index]);
+		const bool bOperationSuccess = SCCProvider->TryToDownloadFileFromBackgroundThread(DownloadCommand, DepotPaths);
+#endif
 
-				Requests[Index].SetPayload(FCompressedBuffer::FromCompressed(Buffer));
-			}
+		// Check and assign what payloads we did find, even if the download command failed we might have been able to download some
+		// files and since we found them we might as well return them.
+		for(int32 Index = 0; Index < Requests.Num(); ++Index)
+		{
+			// The payload was created by FCompressedBuffer::Compress so we can return it as a FCompressedBuffer.
+			FSharedBuffer Buffer = DownloadCommand->GetFileData(DepotPaths[Index]);
+			Requests[Index].SetPayload(FCompressedBuffer::FromCompressed(Buffer));
+		}
 
-			break;
+		if (bOperationSuccess)
+		{
+			return true;
 		}
-#endif
 
 		// If this was the first try then check to see if the error being returns is that the file does not exist
 		// in the depot. If it does not exist then there is no point in us retrying and we can error out at this point.
-		if (Retries == 0 && IsDepotFileMissing(DownloadCommand->GetResultInfo()))
+		if (Retries == 0 && !ShouldRetryOperation(DownloadCommand->GetResultInfo()))
 		{
-			break;
+			return false;
 		}
 
 		Retries++;
 	}
 
-	// TODO -  Should we error here if the payload failed? 
-
-	return true;
-#endif
+	return false;
 }
 
 bool FSourceControlBackend::DoesPayloadExist(const FIoHash& Id)
