Fix source control automated tests that were failing

diff --git a/Engine/Plugins/Developer/PerforceSourceControl/Source/PerforceSourceControl/Private/PerforceSourceControlState.cpp b/Engine/Plugins/Developer/PerforceSourceControl/Source/PerforceSourceControl/Private/PerforceSourceControlState.cpp
--- a/Engine/Plugins/Developer/PerforceSourceControl/Source/PerforceSourceControl/Private/PerforceSourceControlState.cpp
+++ b/Engine/Plugins/Developer/PerforceSourceControl/Source/PerforceSourceControl/Private/PerforceSourceControlState.cpp
@@ -297,7 +297,7 @@ bool FPerforceSourceControlState::IsCurrent() const
 
 bool FPerforceSourceControlState::IsSourceControlled() const
 {
-	return State != EPerforceState::NotInDepot && State != EPerforceState::NotUnderClientRoot;
+	return !IsUnknown() && State != EPerforceState::NotInDepot && State != EPerforceState::NotUnderClientRoot;
 }
 
 bool FPerforceSourceControlState::IsAdded() const
