Fix the perforce API reporting that files are not under source control when polled via the depot path.

#rnx

- If we queried the state of a file using the depot path it would parse correctly but internally record the state as 'DontCare' meaning that it has no special state but the change in 23468037 causes files with that state to be classed as not under source control.
- The root problem that 23468037 was fixing was that when querying a file that does not exist (using client paths) we expect the file state to be set to 'NotInDepot' but for a file in a unit test this state was coming back as 'DontCare'. This was due to the imaginary file was being queried as a depot location (Engine root) where most users do not have permission to add things, so instead of an error message indicating that the file does not exist, the server was returning an error relating to permissions.
- We can fix this by checking for the permission error and adding that to the list of errors that result in us setting the 'NotInDepot' state. Which in turn allows us to remove the change from 23468037 that was blocking the depot path queries.
- There is an argument to be made that we should not allow file states to be queried via depot paths and doing so should raise an error. To allow that we would need to extend the source control api with a way to query files server state which would be out of scope for a release hot fix.

diff --git a/Engine/Plugins/Developer/PerforceSourceControl/Source/PerforceSourceControl/Private/PerforceSourceControlOperations.cpp b/Engine/Plugins/Developer/PerforceSourceControl/Source/PerforceSourceControl/Private/PerforceSourceControlOperations.cpp
--- a/Engine/Plugins/Developer/PerforceSourceControl/Source/PerforceSourceControl/Private/PerforceSourceControlOperations.cpp
+++ b/Engine/Plugins/Developer/PerforceSourceControl/Source/PerforceSourceControl/Private/PerforceSourceControlOperations.cpp
@@ -1397,7 +1397,7 @@ static void ParseUpdateStatusResults(const FP4RecordSet& InRecords, const TArray
 	for (const FText& Error : ErrorMessages)
 	{
 		//@todo P4 could be returning localized error messages
-		int32 NoSuchFilePos = Error.ToString().Find(TEXT(" - no such file(s).\n"), ESearchCase::IgnoreCase, ESearchDir::FromStart);
+		const int32 NoSuchFilePos = Error.ToString().Find(TEXT(" - no such file(s).\n"), ESearchCase::IgnoreCase, ESearchDir::FromStart);
 		if(NoSuchFilePos != INDEX_NONE)
 		{
 			// found an error about a file that is not in the depot
@@ -1409,7 +1409,19 @@ static void ParseUpdateStatusResults(const FP4RecordSet& InRecords, const TArray
 		}
 
 		//@todo P4 could be returning localized error messages
-		int32 NotUnderClientRootPos = Error.ToString().Find(TEXT("' is not under client's root"), ESearchCase::IgnoreCase, ESearchDir::FromStart);
+		const int32 ProtectedPos = Error.ToString().Find(TEXT(" - protected namespace - access denied.\n"), ESearchCase::IgnoreCase, ESearchDir::FromStart);
+		if (ProtectedPos != INDEX_NONE)
+		{
+			// found an error about a file that is not in the depot
+			FString FullPath(Error.ToString().Left(ProtectedPos));
+			FPaths::NormalizeFilename(FullPath);
+			OutStates.Add(FPerforceSourceControlState(FullPath));
+			FPerforceSourceControlState& State = OutStates.Last();
+			State.State = EPerforceState::NotInDepot;
+		}
+
+		//@todo P4 could be returning localized error messages
+		const int32 NotUnderClientRootPos = Error.ToString().Find(TEXT("' is not under client's root"), ESearchCase::IgnoreCase, ESearchDir::FromStart);
 		if(NotUnderClientRootPos != INDEX_NONE)
 		{
 			// found an error about a file that is not under the client root
diff --git a/Engine/Plugins/Developer/PerforceSourceControl/Source/PerforceSourceControl/Private/PerforceSourceControlState.cpp b/Engine/Plugins/Developer/PerforceSourceControl/Source/PerforceSourceControl/Private/PerforceSourceControlState.cpp
--- a/Engine/Plugins/Developer/PerforceSourceControl/Source/PerforceSourceControl/Private/PerforceSourceControlState.cpp
+++ b/Engine/Plugins/Developer/PerforceSourceControl/Source/PerforceSourceControl/Private/PerforceSourceControlState.cpp
@@ -297,7 +297,7 @@ bool FPerforceSourceControlState::IsCurrent() const
 
 bool FPerforceSourceControlState::IsSourceControlled() const
 {
-	return !IsUnknown() && State != EPerforceState::NotInDepot && State != EPerforceState::NotUnderClientRoot;
+	return State != EPerforceState::NotInDepot && State != EPerforceState::NotUnderClientRoot;
 }
 
 bool FPerforceSourceControlState::IsAdded() const
