#UE Remove some redundant messages from the DownloadFile p4 operation, similar to how other operations behave

diff --git a/Engine/Plugins/Developer/PerforceSourceControl/Source/PerforceSourceControl/Private/PerforceSourceControlOperations.cpp b/Engine/Plugins/Developer/PerforceSourceControl/Source/PerforceSourceControl/Private/PerforceSourceControlOperations.cpp
--- a/Engine/Plugins/Developer/PerforceSourceControl/Source/PerforceSourceControl/Private/PerforceSourceControlOperations.cpp
+++ b/Engine/Plugins/Developer/PerforceSourceControl/Source/PerforceSourceControl/Private/PerforceSourceControlOperations.cpp
@@ -3071,6 +3071,11 @@ bool FPerforceDownloadFileWorker::Execute(FPerforceSourceControlCommand& InComma
 																	FOnIsCancelled::CreateRaw(&InCommand, &FPerforceSourceControlCommand::IsCanceled), 
 																	InCommand.bConnectionDropped, Flags);
 			
+			RemoveRedundantErrors(InCommand, TEXT(" - no such file(s)."));
+			RemoveRedundantErrors(InCommand, TEXT(" - file(s) not on client"));
+			RemoveRedundantErrors(InCommand, TEXT("' is not under client's root '"));
+			RemoveRedundantErrors(InCommand, TEXT(" - no file(s) at that changelist number."));
+
 			if (!InCommand.ResultInfo.ErrorMessages.IsEmpty() || (InCommand.Files.Num() != FilesData.Num()))
 			{
 				InCommand.bCommandSuccessful = false;
@@ -3108,6 +3113,11 @@ bool FPerforceDownloadFileWorker::Execute(FPerforceSourceControlCommand& InComma
 																		FOnIsCancelled::CreateRaw(&InCommand, &FPerforceSourceControlCommand::IsCanceled), 
 																		InCommand.bConnectionDropped, Flags);
 
+				RemoveRedundantErrors(InCommand, TEXT(" - no such file(s)."));
+				RemoveRedundantErrors(InCommand, TEXT(" - file(s) not on client"));
+				RemoveRedundantErrors(InCommand, TEXT("' is not under client's root '"));
+				RemoveRedundantErrors(InCommand, TEXT(" - no file(s) at that changelist number."));
+
 				if (InCommand.bCommandSuccessful)
 				{
 					const FString FinalTargetPath = TargetDirectory / FPaths::GetCleanFilename(TargetFilePath);
