[Engine/Plugins]
* Added IWYUSupport.KeepAsIsForNow on a bunch of plugins that requires work to be able to run iwyu

diff --git a/Engine/Plugins/Developer/PerforceSourceControl/Source/PerforceSourceControl/PerforceSourceControl.Build.cs b/Engine/Plugins/Developer/PerforceSourceControl/Source/PerforceSourceControl/PerforceSourceControl.Build.cs
--- a/Engine/Plugins/Developer/PerforceSourceControl/Source/PerforceSourceControl/PerforceSourceControl.Build.cs
+++ b/Engine/Plugins/Developer/PerforceSourceControl/Source/PerforceSourceControl/PerforceSourceControl.Build.cs
@@ -6,6 +6,8 @@ public class PerforceSourceControl : ModuleRules
 {
 	public PerforceSourceControl(ReadOnlyTargetRules Target) : base(Target)
 	{
+		IWYUSupport = IWYUSupport.KeepAsIsForNow;
+
         PrivateDependencyModuleNames.AddRange(
             new string[] {
 				"Core",
diff --git a/Engine/Plugins/Experimental/Mutable/Source/MutableRuntime/MutableRuntime.Build.cs b/Engine/Plugins/Experimental/Mutable/Source/MutableRuntime/MutableRuntime.Build.cs
--- a/Engine/Plugins/Experimental/Mutable/Source/MutableRuntime/MutableRuntime.Build.cs
+++ b/Engine/Plugins/Experimental/Mutable/Source/MutableRuntime/MutableRuntime.Build.cs
@@ -11,6 +11,7 @@ namespace UnrealBuildTool.Rules
 			ShortName = "MuR";
 
 			DefaultBuildSettings = BuildSettingsVersion.V2;
+			IWYUSupport = IWYUSupport.KeepAsIsForNow;
 			//bUseUnity = false;
 
 			PublicIncludePaths.Add(Path.Combine(ModuleDirectory, "Public"));
diff --git a/Engine/Plugins/Experimental/Mutable/Source/MutableTools/MutableTools.Build.cs b/Engine/Plugins/Experimental/Mutable/Source/MutableTools/MutableTools.Build.cs
--- a/Engine/Plugins/Experimental/Mutable/Source/MutableTools/MutableTools.Build.cs
+++ b/Engine/Plugins/Experimental/Mutable/Source/MutableTools/MutableTools.Build.cs
@@ -11,6 +11,7 @@ namespace UnrealBuildTool.Rules
 			ShortName = "MuT";
 
 			DefaultBuildSettings = BuildSettingsVersion.V2;
+			IWYUSupport = IWYUSupport.KeepAsIsForNow;
 			//bUseUnity = false;
 
 			// 
