Fix the perforce source control api returning invalid file paths when the users client spec root ends with a '/'or '\'

#rnx

- CL 23859876 introduced optimizations that consider a filepath to be invalid if the path contains consecutive path separators (DirA//DirB for example instead of DirA/DirB) and although no issues were seen when testing we fairly quickly ran into a user facing errors when trying to submit work from the editor.
- The users client spec root had a '\' at the end of it where as most internal devs do not. This in turn caused problems when parsing file info from source control as the code was assuming that the client spec root had no trailing separator and was adding one.
- We now normalize the client spec root after we parse it via a call to FPaths::NormalizeDirectoryName which should enforce that the root path match all our usual expectations.

diff --git a/Engine/Plugins/Developer/PerforceSourceControl/Source/PerforceSourceControl/Private/PerforceConnection.cpp b/Engine/Plugins/Developer/PerforceSourceControl/Source/PerforceSourceControl/Private/PerforceConnection.cpp
--- a/Engine/Plugins/Developer/PerforceSourceControl/Source/PerforceSourceControl/Private/PerforceConnection.cpp
+++ b/Engine/Plugins/Developer/PerforceSourceControl/Source/PerforceSourceControl/Private/PerforceConnection.cpp
@@ -1257,8 +1257,7 @@ void FPerforceConnection::EstablishConnection(const FPerforceConnectionInfo& InC
 				UE_LOG(LogSourceControl, Verbose, TEXT(" ... getting clientroot" ));
 				ClientRoot = Records[0](TEXT("clientRoot"));
 
-				//make sure all slashes point the same way
-				ClientRoot = ClientRoot.Replace(TEXT("\\"), TEXT("/"));
+				FPaths::NormalizeDirectoryName(ClientRoot);
 			}
 		}
 	}
