[Backout] - Reverted Changes to SPerforceSourceControlSettings made in CL23857698

diff --git a/Engine/Plugins/Developer/PerforceSourceControl/Source/PerforceSourceControl/Private/SPerforceSourceControlSettings.cpp b/Engine/Plugins/Developer/PerforceSourceControl/Source/PerforceSourceControl/Private/SPerforceSourceControlSettings.cpp
--- a/Engine/Plugins/Developer/PerforceSourceControl/Source/PerforceSourceControl/Private/SPerforceSourceControlSettings.cpp
+++ b/Engine/Plugins/Developer/PerforceSourceControl/Source/PerforceSourceControl/Private/SPerforceSourceControlSettings.cpp
@@ -218,8 +218,7 @@ void SPerforceSourceControlSettings::Construct(const FArguments& InArgs, FPerfor
 
 FString SPerforceSourceControlSettings::GetPassword()
 {
-	const bool bIsValidThread = IsInGameThread() || IsInSlateThread();
-	if(PasswordTextBox.IsValid() && bIsValidThread)
+	if(PasswordTextBox.IsValid())
 	{
 		return PasswordTextBox.Pin()->GetText().ToString();
 	}
diff --git a/Engine/Plugins/Developer/PerforceSourceControl/Source/PerforceSourceControl/Private/SPerforceSourceControlSettings.h b/Engine/Plugins/Developer/PerforceSourceControl/Source/PerforceSourceControl/Private/SPerforceSourceControlSettings.h
--- a/Engine/Plugins/Developer/PerforceSourceControl/Source/PerforceSourceControl/Private/SPerforceSourceControlSettings.h
+++ b/Engine/Plugins/Developer/PerforceSourceControl/Source/PerforceSourceControl/Private/SPerforceSourceControlSettings.h
@@ -43,10 +43,7 @@ public:
 
 	void Construct(const FArguments& InArgs, FPerforceSourceControlProvider* InSCCProvider);
 
-	/** Get the currently entered password
-	 *
-	 *  Note: This can only be called from a game or slate thread. All other threads will return an empty string.
-	 */
+	/** Get the currently entered password */
 	static FString GetPassword();
 
 	virtual void Tick( const FGeometry& AllottedGeometry, const double InCurrentTime, const float InDeltaTime );
