PerforceSourceControl: Use standard logging instead of an ensure to report unhandled file state

diff --git a/Engine/Plugins/Developer/PerforceSourceControl/Source/PerforceSourceControl/Private/PerforceSourceControlOperations.cpp b/Engine/Plugins/Developer/PerforceSourceControl/Source/PerforceSourceControl/Private/PerforceSourceControlOperations.cpp
--- a/Engine/Plugins/Developer/PerforceSourceControl/Source/PerforceSourceControl/Private/PerforceSourceControlOperations.cpp
+++ b/Engine/Plugins/Developer/PerforceSourceControl/Source/PerforceSourceControl/Private/PerforceSourceControlOperations.cpp
@@ -12,6 +12,7 @@
 #include "Misc/EngineVersion.h"
 #include "Misc/Paths.h"
 #include "PerforceConnection.h"
+#include "PerforceMessageLog.h"
 #include "PerforceSourceControlChangeStatusOperation.h"
 #include "PerforceSourceControlChangelistState.h"
 #include "PerforceSourceControlCommand.h"
@@ -1355,7 +1356,11 @@ static void ParseUpdateStatusResults(const FP4RecordSet& InRecords, const TArray
 					if (!ClientRecord.Contains(*VarName))
 					{
 						// No more revisions
-						ensureMsgf( ResolveActionNumber > 0, TEXT("Resolve is pending but no resolve actions for file %s"), *FileName );
+						if (ResolveActionNumber == 0)
+						{
+							FTSMessageLog SourceControlLog("SourceControl");
+							SourceControlLog.Error(FText::Format(LOCTEXT("P4Operation_ResolveWithoutAction", "Resolve is pending but no resolve actions for file {0}"), FText::FromString(FileName)));
+						}
 						break;
 					}
 
@@ -1363,8 +1368,11 @@ static void ParseUpdateStatusResults(const FP4RecordSet& InRecords, const TArray
 					const FString& ResolveBaseFile = ClientRecord(VarName);
 					VarName = FString::Printf(TEXT("resolveFromFile%d"), ResolveActionNumber);
 					const FString& ResolveFromFile = ClientRecord(VarName);
-					if(!ensureMsgf( ResolveFromFile == ResolveBaseFile, TEXT("Text cannot resolve %s with %s, we do not support cross file merging"), *ResolveBaseFile, *ResolveFromFile ) )
+
+					if(ResolveFromFile != ResolveBaseFile)
 					{
+						FTSMessageLog SourceControlLog("SourceControl");
+						SourceControlLog.Error(FText::Format(LOCTEXT("P4Operation_UnsupportedCrossFileMerge", "Text cannot resolve {0} with {1}, we do not support cross file merging"), FText::FromString(ResolveBaseFile), FText::FromString(ResolveFromFile)));
 						break;
 					}
 
