Add an idle timeout on perforce persistent connections

A connection to the perforce server can consume a large amount of memory on the server.  This change is to add a timeout to any persistent connections
in order to conserve this memory.
To do so, a connection tracks the platform time of the last successful communication.  The provider checks the persistent connection it has each frame and if
the last successful connect exceeded the timeout, then the connection is closed.  A new persistent connection will be established next time one is required.

The timeout is set as a hardcoded value for now.

diff --git a/Engine/Plugins/Developer/PerforceSourceControl/Source/PerforceSourceControl/Private/PerforceConnection.cpp b/Engine/Plugins/Developer/PerforceSourceControl/Source/PerforceSourceControl/Private/PerforceConnection.cpp
--- a/Engine/Plugins/Developer/PerforceSourceControl/Source/PerforceSourceControl/Private/PerforceConnection.cpp
+++ b/Engine/Plugins/Developer/PerforceSourceControl/Source/PerforceSourceControl/Private/PerforceConnection.cpp
@@ -610,6 +610,7 @@ static bool CheckUnicodeStatus(ClientApi& P4Client, bool& bIsUnicodeServer, TArr
 FPerforceConnection::FPerforceConnection(const FPerforceConnectionInfo& InConnectionInfo, FPerforceSourceControlProvider& InSCCProvider)
 	: bEstablishedConnection(false)
 	, bIsUnicode(false)
+	, LatestCommunicateTime(0.0)
 	, SCCProvider(InSCCProvider)
 {
 	EstablishConnection(InConnectionInfo);
@@ -979,9 +980,14 @@ void FPerforceConnection::Disconnect()
 	}
 }
 
+double FPerforceConnection::GetLatestCommuncationTime() const
+{
+	return LatestCommunicateTime;
+}
+
 bool FPerforceConnection::RunCommand(	const FString& InCommand, const TArray<FString>& InParameters, FP4RecordSet& OutRecordSet, 
-										TArray<FSharedBuffer>* OutData, TArray<FText>& OutErrorMessage, 
-										FOnIsCancelled InIsCancelled, bool& OutConnectionDropped, ERunCommandFlags RunFlags)
+                                        TArray<FSharedBuffer>* OutData, TArray<FText>& OutErrorMessage, 
+                                        FOnIsCancelled InIsCancelled, bool& OutConnectionDropped, ERunCommandFlags RunFlags)
 {
 	TRACE_CPUPROFILER_EVENT_SCOPE_TEXT(*FString::Printf(TEXT("FPerforceConnection::RunCommand_%s"), *InCommand));
 
@@ -1081,6 +1087,11 @@ bool FPerforceConnection::RunCommand(	const FString& InCommand, const TArray<FSt
 		}
 	}
 
+	if (!OutConnectionDropped)
+	{
+		LatestCommunicateTime = FPlatformTime::Seconds();
+	}
+
 	return OutRecordSet.Num() > 0;
 }
 
diff --git a/Engine/Plugins/Developer/PerforceSourceControl/Source/PerforceSourceControl/Private/PerforceConnection.h b/Engine/Plugins/Developer/PerforceSourceControl/Source/PerforceSourceControl/Private/PerforceConnection.h
--- a/Engine/Plugins/Developer/PerforceSourceControl/Source/PerforceSourceControl/Private/PerforceConnection.h
+++ b/Engine/Plugins/Developer/PerforceSourceControl/Source/PerforceSourceControl/Private/PerforceConnection.h
@@ -102,6 +102,9 @@ public:
 	/** If the connection is valid, disconnect from the server */
 	void Disconnect();
 
+	/** Returns the platform time of the last time the connection successfully communicated with the server */
+	double GetLatestCommuncationTime() const;
+
 	/**
 	 * Runs internal perforce command, catches exceptions, returns results
 	 */
@@ -170,6 +173,9 @@ public:
 	/** Is this a connection to a unicode server? */ 
 	bool bIsUnicode;
 
+	/** A record of the last time the connection was used */
+	double LatestCommunicateTime;
+
 private:
 	/** The source control provider that the connection is working from */
 	FPerforceSourceControlProvider& SCCProvider;
diff --git a/Engine/Plugins/Developer/PerforceSourceControl/Source/PerforceSourceControl/Private/PerforceSourceControlProvider.cpp b/Engine/Plugins/Developer/PerforceSourceControl/Source/PerforceSourceControl/Private/PerforceSourceControlProvider.cpp
--- a/Engine/Plugins/Developer/PerforceSourceControl/Source/PerforceSourceControl/Private/PerforceSourceControlProvider.cpp
+++ b/Engine/Plugins/Developer/PerforceSourceControl/Source/PerforceSourceControl/Private/PerforceSourceControlProvider.cpp
@@ -25,6 +25,15 @@ static FName ProviderName("Perforce");
 
 #define LOCTEXT_NAMESPACE "PerforceSourceControl"
 
+namespace 
+{
+	/** Amount of seconds an idle persistent connection can remain open before the provider closes it.
+	 * This lowers the number of open connections to a perforce server across a studio
+	 * to conserve memory usage of the server.
+	 */
+	constexpr double IdleConnectionDisconnectSeconds = 60 * 60.0; // 1 hour
+}
+
 FPerforceSourceControlProvider::FPerforceSourceControlProvider()
 	: PerforceSCCSettings(*this, FStringView())
 	, InitialSettings(FSourceControlInitSettings::EBehavior::OverrideExisting)
@@ -606,12 +615,28 @@ void FPerforceSourceControlProvider::Tick()
 			break;
 		}
 	}
-
+	
 	if(bStatesUpdated)
 	{
 		TRACE_CPUPROFILER_EVENT_SCOPE(FPerforceSourceControlProvider::Tick::BroadcastStateUpdate);
 		OnSourceControlStateChanged.Broadcast();
 	}
+
+	if (PersistentConnection)
+	{
+		const double Now = FPlatformTime::Seconds();
+
+		const double ElapsedSinceLastComm = Now - PersistentConnection->GetLatestCommuncationTime();
+
+		if (ElapsedSinceLastComm > IdleConnectionDisconnectSeconds)
+		{
+			UE_LOG(LogSourceControl, Display,
+				TEXT("Persisent perforce connection has not been used in %0.f seconds. Dropping connection"),
+				ElapsedSinceLastComm);
+
+			Close();
+		}
+	}
 }
 
 static void ParseGetLabelsResults(FPerforceSourceControlProvider& InSourceControlProvider, const FP4RecordSet& InRecords, TArray< TSharedRef<ISourceControlLabel> >& OutLabels)
