Add CVar to change perforce idle connection timeout
* Aids in testing this feature
* Use CMD:  SourceControl.Perforce.IdleConnectionDisconnectSeconds = <seconds>

diff --git a/Engine/Plugins/Developer/PerforceSourceControl/Source/PerforceSourceControl/Private/PerforceSourceControlProvider.cpp b/Engine/Plugins/Developer/PerforceSourceControl/Source/PerforceSourceControl/Private/PerforceSourceControlProvider.cpp
--- a/Engine/Plugins/Developer/PerforceSourceControl/Source/PerforceSourceControl/Private/PerforceSourceControlProvider.cpp
+++ b/Engine/Plugins/Developer/PerforceSourceControl/Source/PerforceSourceControl/Private/PerforceSourceControlProvider.cpp
@@ -31,7 +31,13 @@ namespace
 	 * This lowers the number of open connections to a perforce server across a studio
 	 * to conserve memory usage of the server.
 	 */
-	constexpr double IdleConnectionDisconnectSeconds = 60 * 60.0; // 1 hour
+	float IdleConnectionDisconnectSeconds = 60 * 60.0; // 1 hour
+
+	static FAutoConsoleVariableRef CVarIdleConnectionDisconnectSeconds(
+		TEXT("SourceControl.Perforce.IdleConnectionDisconnectSeconds"),
+		IdleConnectionDisconnectSeconds,
+		TEXT("The number of seconds a perforce connection will be kept open without activity before being automatically disconnected"),
+		ECVF_Default);
 }
 
 FPerforceSourceControlProvider::FPerforceSourceControlProvider()
@@ -632,11 +638,11 @@ void FPerforceSourceControlProvider::Tick()
 		const double Now = FPlatformTime::Seconds();
 
 		const double ElapsedSinceLastComm = Now - PersistentConnection->GetLatestCommuncationTime();
-
+		
 		if (ElapsedSinceLastComm > IdleConnectionDisconnectSeconds)
 		{
 			UE_LOG(LogSourceControl, Display,
-				TEXT("Persisent perforce connection has not been used in %0.f seconds. Dropping connection"),
+				TEXT("Persistent perforce connection has not been used in %0.f seconds. Dropping connection"),
 				ElapsedSinceLastComm);
 
 			Close();
