Remove unused perforce connection from FPerforceGetFileWorker

#rnx

- FPerforceGetFileWorker::Execute doesn't actually do any perforce work directly, it creates a new FPerforceSourceControlRevision and lets that object do the perforce work. So the connection that FPerforceGetFileWorker was making was never used and FPerforceSourceControlRevision would then create a second connection.

diff --git a/Engine/Plugins/Developer/PerforceSourceControl/Source/PerforceSourceControl/Private/PerforceSourceControlOperations.cpp b/Engine/Plugins/Developer/PerforceSourceControl/Source/PerforceSourceControl/Private/PerforceSourceControlOperations.cpp
--- a/Engine/Plugins/Developer/PerforceSourceControl/Source/PerforceSourceControl/Private/PerforceSourceControlOperations.cpp
+++ b/Engine/Plugins/Developer/PerforceSourceControl/Source/PerforceSourceControl/Private/PerforceSourceControlOperations.cpp
@@ -3364,11 +3364,8 @@ bool FPerforceGetFileWorker::Execute(FPerforceSourceControlCommand& InCommand)
 {
 	TRACE_CPUPROFILER_EVENT_SCOPE(FPerforceGetChangelistDetailsWorker::FPerforceGetFileWorker);
 
-	FScopedPerforceConnection ScopedConnection(InCommand);
-
-	if (!InCommand.IsCanceled() && ScopedConnection.IsValid())
+	if (!InCommand.IsCanceled())
 	{
-		FPerforceConnection& Connection = ScopedConnection.GetConnection();
 		TSharedRef<FGetFile, ESPMode::ThreadSafe> Operation = StaticCastSharedRef<FGetFile>(InCommand.Operation);
 
 		TSharedRef<FPerforceSourceControlRevision, ESPMode::ThreadSafe> Revision = MakeShareable(new FPerforceSourceControlRevision(GetSCCProvider()));
