[Attempted BugFix] Editor crashes when loading certain CLs in the Revision Changelist window of the UEQA_SC_P4V project
made PerforceSourceControlRevision.cpp temp file naming multithread friendly

diff --git a/Engine/Plugins/Developer/PerforceSourceControl/Source/PerforceSourceControl/Private/PerforceSourceControlRevision.cpp b/Engine/Plugins/Developer/PerforceSourceControl/Source/PerforceSourceControl/Private/PerforceSourceControlRevision.cpp
--- a/Engine/Plugins/Developer/PerforceSourceControl/Source/PerforceSourceControl/Private/PerforceSourceControlRevision.cpp
+++ b/Engine/Plugins/Developer/PerforceSourceControl/Source/PerforceSourceControl/Private/PerforceSourceControlRevision.cpp
@@ -58,8 +58,8 @@ bool FPerforceSourceControlRevision::Get( FString& InOutFilename, EConcurrency::
 		}
 		else
 		{
-			static int32 TempFileCount = 0;
-			FString TempFileName = FString::Printf(TEXT("%sTemp-%d-Rev-%s-%s"), *FPaths::DiffDir(), TempFileCount++, *RevString, *FPaths::GetCleanFilename(FileName));
+			const FString File = FString::Printf(TEXT("%s-Rev-%s-"), *FPaths::GetBaseFilename(FileName), *RevString);
+			const FString TempFileName = FPaths::CreateTempFilename(*FPaths::DiffDir(), *File, *FPaths::GetExtension(FileName));
 			AbsoluteFileName = FPaths::ConvertRelativePathToFull(TempFileName);
 		}
 
@@ -209,9 +209,9 @@ bool FPerforceSourceControlRevision::GetAnnotated( FString& InOutFilename ) cons
 		}
 		else
 		{
-			FString RevString = (RevisionNumber < 0) ? TEXT("head") : FString::Printf(TEXT("%d"), RevisionNumber);
-			static int32 TempFileCount = 0;
-			FString TempFileName = FString::Printf(TEXT("%sAnnotated-%d-Rev-%s-%s"), *FPaths::DiffDir(), TempFileCount++, *RevString, *FPaths::GetCleanFilename(FileName));
+			const FString RevString = (RevisionNumber < 0) ? TEXT("head") : FString::Printf(TEXT("%d"), RevisionNumber);
+			const FString AnnotatedName = FString::Printf(TEXT("Annotated-%s-Rev-%s-"), *FPaths::GetBaseFilename(FileName), *RevString);
+			const FString TempFileName = FPaths::CreateTempFilename(*FPaths::DiffDir(), *AnnotatedName, *FPaths::GetExtension(FileName));
 			AbsoluteFileName = FPaths::ConvertRelativePathToFull(TempFileName);
 		}
 
