[BugFix] Fixed Diff temp file Extensions missing '.' before them which caused .uassets to not be loadable

diff --git a/Engine/Plugins/Developer/PerforceSourceControl/Source/PerforceSourceControl/Private/PerforceSourceControlRevision.cpp b/Engine/Plugins/Developer/PerforceSourceControl/Source/PerforceSourceControl/Private/PerforceSourceControlRevision.cpp
--- a/Engine/Plugins/Developer/PerforceSourceControl/Source/PerforceSourceControl/Private/PerforceSourceControlRevision.cpp
+++ b/Engine/Plugins/Developer/PerforceSourceControl/Source/PerforceSourceControl/Private/PerforceSourceControlRevision.cpp
@@ -59,7 +59,8 @@ bool FPerforceSourceControlRevision::Get( FString& InOutFilename, EConcurrency::
 		else
 		{
 			const FString File = FString::Printf(TEXT("%s-Rev-%s-"), *FPaths::GetBaseFilename(FileName), *RevString);
-			const FString TempFileName = FPaths::CreateTempFilename(*FPaths::DiffDir(), *File, *FPaths::GetExtension(FileName));
+			const FString Extension = TEXT(".") + FPaths::GetExtension(FileName);
+			const FString TempFileName = FPaths::CreateTempFilename(*FPaths::DiffDir(), *File, *Extension);
 			AbsoluteFileName = FPaths::ConvertRelativePathToFull(TempFileName);
 		}
 
@@ -211,7 +212,8 @@ bool FPerforceSourceControlRevision::GetAnnotated( FString& InOutFilename ) cons
 		{
 			const FString RevString = (RevisionNumber < 0) ? TEXT("head") : FString::Printf(TEXT("%d"), RevisionNumber);
 			const FString AnnotatedName = FString::Printf(TEXT("Annotated-%s-Rev-%s-"), *FPaths::GetBaseFilename(FileName), *RevString);
-			const FString TempFileName = FPaths::CreateTempFilename(*FPaths::DiffDir(), *AnnotatedName, *FPaths::GetExtension(FileName));
+			const FString Extension = TEXT(".") + FPaths::GetExtension(FileName);
+			const FString TempFileName = FPaths::CreateTempFilename(*FPaths::DiffDir(), *AnnotatedName, *Extension);
 			AbsoluteFileName = FPaths::ConvertRelativePathToFull(TempFileName);
 		}
 
