Improve the error message given via FPerforceDownloadFileWorker if a depot for one of the requested files cannot be found.

#rnx

- If p4 print is run on a path for a depot that either does not exist on the server or cannot be seen due to permissions then the server will return "-must refer to client" along with the name of the current client spec. This can be quite misleading since p4 print does not actually need a client.
-- To improve things for the user we can replace the error message provided by the p4 server with "no such depot"

diff --git a/Engine/Plugins/Developer/PerforceSourceControl/Source/PerforceSourceControl/Private/PerforceSourceControlOperations.cpp b/Engine/Plugins/Developer/PerforceSourceControl/Source/PerforceSourceControl/Private/PerforceSourceControlOperations.cpp
--- a/Engine/Plugins/Developer/PerforceSourceControl/Source/PerforceSourceControl/Private/PerforceSourceControlOperations.cpp
+++ b/Engine/Plugins/Developer/PerforceSourceControl/Source/PerforceSourceControl/Private/PerforceSourceControlOperations.cpp
@@ -176,6 +176,49 @@ static void RemoveRedundantErrors(FPerforceSourceControlCommand& InCommand, cons
 	}
 }
 
+/**
+ * Occasionally some error messages returned by the server can be a bit misleading depending on the context
+ * of the original command. This utility allows us to replace an error that we feel is misleading with a
+ * better version.
+ *
+ * If the misleading error is found we don't just replace that exact string, we remove everything from the
+ * string, starting at the misleading error and then append the correct error on the end. This is because
+ * the error messages can often contain variable info (such as the port or client) which can be hard to
+ * match the misleading error string against.
+ */
+static void ReplaceErrors(FPerforceSourceControlCommand& InCommand, const TCHAR* InMisleadingError, const TCHAR* InCorrectedError, bool bMoveToInfo)
+{
+	TArray<FText>& ErrorMessages = InCommand.ResultInfo.ErrorMessages;
+
+	for (int32 Index = 0; Index < ErrorMessages.Num(); ++Index)
+	{
+		const FString Error = ErrorMessages[Index].ToString();
+
+		const int32 Pos = Error.Find(InMisleadingError);
+		if (Pos != INDEX_NONE)
+		{
+			TStringBuilder<128> NewError;
+			NewError << FStringView(Error).Left(Pos) << InCorrectedError << TEXT("\n");
+
+			if (bMoveToInfo)
+			{
+				InCommand.ResultInfo.InfoMessages.Add(FText::FromString(NewError.ToString()));
+				ErrorMessages.RemoveAt(Index);
+				--Index;
+			}
+			else
+			{
+				ErrorMessages[Index] = FText::FromString(NewError.ToString());
+			}
+		}
+	}
+
+	if (InCommand.ResultInfo.ErrorMessages.IsEmpty())
+	{
+		InCommand.bCommandSuccessful = true;
+	}
+}
+
 /** Simple parsing of a record set into strings, one string per record */
 static void ParseRecordSet(const FP4RecordSet& InRecords, TArray<FText>& OutResults)
 {
@@ -3079,10 +3122,12 @@ bool FPerforceDownloadFileWorker::Execute(FPerforceSourceControlCommand& InComma
 																	FOnIsCancelled::CreateRaw(&InCommand, &FPerforceSourceControlCommand::IsCanceled), 
 																	InCommand.bConnectionDropped, Flags);
 			
-			RemoveRedundantErrors(InCommand, TEXT(" - no such file(s)."));
-			RemoveRedundantErrors(InCommand, TEXT(" - file(s) not on client"));
-			RemoveRedundantErrors(InCommand, TEXT("' is not under client's root '"));
-			RemoveRedundantErrors(InCommand, TEXT(" - no file(s) at that changelist number."));
+			const bool bMoveToInfo = true;
+			RemoveRedundantErrors(InCommand, TEXT(" - no such file(s)."), bMoveToInfo);
+			RemoveRedundantErrors(InCommand, TEXT(" - file(s) not on client"), bMoveToInfo);
+			RemoveRedundantErrors(InCommand, TEXT("' is not under client's root '"), bMoveToInfo);
+			RemoveRedundantErrors(InCommand, TEXT(" - no file(s) at that changelist number."), bMoveToInfo);
+			ReplaceErrors(InCommand, TEXT(" - must refer to client"), TEXT(" - no such depot!"), bMoveToInfo);
 
 			if (!InCommand.ResultInfo.ErrorMessages.IsEmpty() || (InCommand.Files.Num() != FilesData.Num()))
 			{
@@ -3121,10 +3166,12 @@ bool FPerforceDownloadFileWorker::Execute(FPerforceSourceControlCommand& InComma
 																		FOnIsCancelled::CreateRaw(&InCommand, &FPerforceSourceControlCommand::IsCanceled), 
 																		InCommand.bConnectionDropped, Flags);
 
-				RemoveRedundantErrors(InCommand, TEXT(" - no such file(s)."));
-				RemoveRedundantErrors(InCommand, TEXT(" - file(s) not on client"));
-				RemoveRedundantErrors(InCommand, TEXT("' is not under client's root '"));
-				RemoveRedundantErrors(InCommand, TEXT(" - no file(s) at that changelist number."));
+				const bool bMoveToInfo = true;
+				RemoveRedundantErrors(InCommand, TEXT(" - no such file(s)."), bMoveToInfo);
+				RemoveRedundantErrors(InCommand, TEXT(" - file(s) not on client"), bMoveToInfo);
+				RemoveRedundantErrors(InCommand, TEXT("' is not under client's root '"), bMoveToInfo);
+				RemoveRedundantErrors(InCommand, TEXT(" - no file(s) at that changelist number."), bMoveToInfo);
+				ReplaceErrors(InCommand, TEXT(" - must refer to client"), TEXT(" - no such depot!"), bMoveToInfo);
 
 				if (InCommand.bCommandSuccessful)
 				{
