[CDPR] Added more information about changelist to source control changelists window context menu

+ Expand source control menu context with selected changelist,
+ Expose information about if changelist is default and its source control specific identifier,

diff --git a/Engine/Plugins/Developer/PerforceSourceControl/Source/PerforceSourceControl/Private/PerforceSourceControlChangelist.h b/Engine/Plugins/Developer/PerforceSourceControl/Source/PerforceSourceControl/Private/PerforceSourceControlChangelist.h
--- a/Engine/Plugins/Developer/PerforceSourceControl/Source/PerforceSourceControl/Private/PerforceSourceControlChangelist.h
+++ b/Engine/Plugins/Developer/PerforceSourceControl/Source/PerforceSourceControl/Private/PerforceSourceControlChangelist.h
@@ -37,7 +37,7 @@ public:
 		return ChangelistNumber != InOther.ChangelistNumber;
 	}
 
-	bool IsDefault() const
+	virtual bool IsDefault() const override
 	{
 		return ChangelistNumber == DefaultChangelist.ChangelistNumber;
 	}
@@ -68,6 +68,11 @@ public:
 		return ChangelistNumber;
 	}
 
+	virtual FString GetIdentifier() const override
+	{
+		return ToString();
+	}
+
 public:
 	static const FPerforceSourceControlChangelist DefaultChangelist;
 
diff --git a/Engine/Plugins/Developer/PlasticSourceControl/Source/PlasticSourceControl/Private/PlasticSourceControlChangelist.h b/Engine/Plugins/Developer/PlasticSourceControl/Source/PlasticSourceControl/Private/PlasticSourceControlChangelist.h
--- a/Engine/Plugins/Developer/PlasticSourceControl/Source/PlasticSourceControl/Private/PlasticSourceControlChangelist.h
+++ b/Engine/Plugins/Developer/PlasticSourceControl/Source/PlasticSourceControl/Private/PlasticSourceControlChangelist.h
@@ -37,7 +37,7 @@ public:
 		return ChangelistName != InOther.ChangelistName;
 	}
 
-	bool IsDefault() const
+	virtual bool IsDefault() const override
 	{
 		return ChangelistName == DefaultChangelist.ChangelistName;
 	}
@@ -68,6 +68,11 @@ public:
 		return ChangelistName;
 	}
 
+	virtual FString GetIdentifier() const override
+	{
+		return ChangelistName;
+	}
+
 public:
 	static const FPlasticSourceControlChangelist DefaultChangelist;
 
diff --git a/Engine/Source/Developer/SourceControl/Public/ISourceControlChangelist.h b/Engine/Source/Developer/SourceControl/Public/ISourceControlChangelist.h
--- a/Engine/Source/Developer/SourceControl/Public/ISourceControlChangelist.h
+++ b/Engine/Source/Developer/SourceControl/Public/ISourceControlChangelist.h
@@ -25,4 +25,14 @@ public:
 	 * default Perforce changelist.
 	 */
 	virtual bool CanDelete() const { return true; }
+
+	/**
+	 * Get source control specific identifier for this changelist, for example CL number in Perforce.
+	 */
+	virtual FString GetIdentifier() const { return TEXT(""); }
+
+	/**
+	 * Returns true if this changelist is a default changelist.
+	 */
+	virtual bool IsDefault() const { return false; }
 };
diff --git a/Engine/Source/Editor/SourceControlWindows/Private/SSourceControlChangelists.cpp b/Engine/Source/Editor/SourceControlWindows/Private/SSourceControlChangelists.cpp
--- a/Engine/Source/Editor/SourceControlWindows/Private/SSourceControlChangelists.cpp
+++ b/Engine/Source/Editor/SourceControlWindows/Private/SSourceControlChangelists.cpp
@@ -2865,6 +2865,8 @@ TSharedPtr<SWidget> SSourceControlChangelistsWidget::OnOpenContextMenu()
 	ContextObject->SelectedFiles.Append(SelectedControlledFiles);
 	ContextObject->SelectedFiles.Append(SelectedUncontrolledFiles);
 
+	ContextObject->SelectedChangelist = GetChangelistFromSelection();
+
 	UToolMenu* Menu = ToolMenus->GenerateMenu(MenuName, Context);
 
 	FToolMenuSection& Section = *Menu->FindSection("Source Control");
diff --git a/Engine/Source/Editor/SourceControlWindows/Public/SourceControlMenuContext.h b/Engine/Source/Editor/SourceControlWindows/Public/SourceControlMenuContext.h
--- a/Engine/Source/Editor/SourceControlWindows/Public/SourceControlMenuContext.h
+++ b/Engine/Source/Editor/SourceControlWindows/Public/SourceControlMenuContext.h
@@ -4,6 +4,7 @@
 #include "CoreMinimal.h"
 #include "UObject/ObjectMacros.h"
 #include "UObject/Object.h"
+#include "ISourceControlChangelist.h"
 
 #include "SourceControlMenuContext.generated.h"
 
@@ -24,4 +25,6 @@ public:
 
 	UPROPERTY()
 	TArray<FString> SelectedFiles;
+
+	FSourceControlChangelistPtr SelectedChangelist;
 };
