Fixed mismatched copy constructor/assignment operator for FPerforceSourceControlChangelist.

diff --git a/Engine/Plugins/Developer/PerforceSourceControl/Source/PerforceSourceControl/Private/PerforceSourceControlChangelist.h b/Engine/Plugins/Developer/PerforceSourceControl/Source/PerforceSourceControl/Private/PerforceSourceControlChangelist.h
--- a/Engine/Plugins/Developer/PerforceSourceControl/Source/PerforceSourceControl/Private/PerforceSourceControlChangelist.h
+++ b/Engine/Plugins/Developer/PerforceSourceControl/Source/PerforceSourceControl/Private/PerforceSourceControlChangelist.h
@@ -16,12 +16,6 @@ public:
 	{
 	}
 
-	FPerforceSourceControlChangelist(const FPerforceSourceControlChangelist& InOther)
-		: ChangelistNumber(InOther.ChangelistNumber)
-		, bInitialized(InOther.bInitialized)
-	{
-	}
-
 	virtual bool CanDelete() const override
 	{
 		return !IsDefault();
