Algorithmic optimization in PerforceRevertWorker
* Especially slow when large numbers of files are checked out (in the 25k+ range)
* Existing algorithm scales cubically
* Optimization #1: Changed to take advantage of the fact that filenames are already in a map, so a cheaper lookup to check if file in changelist is part of the to-be-deleted set
* Optimization #2: Another optimization takes advantage of the fact that it is typical the changelist to update is often the one in the to be deleted set.  This is not going to be true for all workflows but helpful anyways.
* Benchmark shows deleting 100 actors from a level with 25k one file per actors checked out in default changelist takes the time to complete the operation from 12s to ~200us. Without optimization #2, this takes about 2s to complete.

diff --git a/Engine/Plugins/Developer/PerforceSourceControl/Source/PerforceSourceControl/Private/PerforceSourceControlOperations.cpp b/Engine/Plugins/Developer/PerforceSourceControl/Source/PerforceSourceControl/Private/PerforceSourceControlOperations.cpp
--- a/Engine/Plugins/Developer/PerforceSourceControl/Source/PerforceSourceControl/Private/PerforceSourceControlOperations.cpp
+++ b/Engine/Plugins/Developer/PerforceSourceControl/Source/PerforceSourceControl/Private/PerforceSourceControlOperations.cpp
@@ -566,9 +566,7 @@ static bool RemoveFilesFromChangelist(const TMap<FString, EPerforceState::Type>&
 {
 	return ChangelistState->Files.RemoveAll([&Results](FSourceControlStateRef& State) -> bool
 		{
-			return Algo::AnyOf(Results, [&State](auto& Result) {
-				return State->GetFilename() == Result.Key;
-				});
+			return Results.Contains(State->GetFilename());
 		}) > 0;
 }
 
@@ -981,12 +979,17 @@ bool FPerforceRevertWorker::UpdateStates() const
 
 	bUpdatedChangelists = ChangelistToUpdate.IsInitialized() && RemoveFilesFromChangelist(OutResults, GetSCCProvider(), ChangelistToUpdate);
 
+	TSharedRef<FPerforceSourceControlChangelistState, ESPMode::ThreadSafe> ChangelistToUpdateState = GetSCCProvider().GetStateInternal(ChangelistToUpdate);
+
 	// Use reverted files to update changelist state.
 	for (TMap<FString, EPerforceState::Type>::TConstIterator It(OutResults); It; ++It)
 	{
 		TSharedRef<FPerforceSourceControlState, ESPMode::ThreadSafe> State = GetSCCProvider().GetStateInternal(It.Key());
 
-		bUpdatedChangelists |= State->Changelist.IsInitialized() && RemoveFilesFromChangelist(OutResults, GetSCCProvider(), State->Changelist);
+		if (State->Changelist != ChangelistToUpdate)
+		{
+			bUpdatedChangelists |= State->Changelist.IsInitialized() && RemoveFilesFromChangelist(OutResults, GetSCCProvider(), State->Changelist);
+		}
 	}
 
 	return bUpdatedCachedStates || bUpdatedChangelists;
