Precautionary undo changelist 25604555

diff --git a/Engine/Plugins/Developer/PerforceSourceControl/Source/PerforceSourceControl/Private/PerforceSourceControlOperations.cpp b/Engine/Plugins/Developer/PerforceSourceControl/Source/PerforceSourceControl/Private/PerforceSourceControlOperations.cpp
--- a/Engine/Plugins/Developer/PerforceSourceControl/Source/PerforceSourceControl/Private/PerforceSourceControlOperations.cpp
+++ b/Engine/Plugins/Developer/PerforceSourceControl/Source/PerforceSourceControl/Private/PerforceSourceControlOperations.cpp
@@ -566,7 +566,9 @@ static bool RemoveFilesFromChangelist(const TMap<FString, EPerforceState::Type>&
 {
 	return ChangelistState->Files.RemoveAll([&Results](FSourceControlStateRef& State) -> bool
 		{
-			return Results.Contains(State->GetFilename());
+			return Algo::AnyOf(Results, [&State](auto& Result) {
+				return State->GetFilename() == Result.Key;
+				});
 		}) > 0;
 }
 
@@ -979,17 +981,12 @@ bool FPerforceRevertWorker::UpdateStates() const
 
 	bUpdatedChangelists = ChangelistToUpdate.IsInitialized() && RemoveFilesFromChangelist(OutResults, GetSCCProvider(), ChangelistToUpdate);
 
-	TSharedRef<FPerforceSourceControlChangelistState, ESPMode::ThreadSafe> ChangelistToUpdateState = GetSCCProvider().GetStateInternal(ChangelistToUpdate);
-
 	// Use reverted files to update changelist state.
 	for (TMap<FString, EPerforceState::Type>::TConstIterator It(OutResults); It; ++It)
 	{
 		TSharedRef<FPerforceSourceControlState, ESPMode::ThreadSafe> State = GetSCCProvider().GetStateInternal(It.Key());
 
-		if (State->Changelist != ChangelistToUpdate)
-		{
-			bUpdatedChangelists |= State->Changelist.IsInitialized() && RemoveFilesFromChangelist(OutResults, GetSCCProvider(), State->Changelist);
-		}
+		bUpdatedChangelists |= State->Changelist.IsInitialized() && RemoveFilesFromChangelist(OutResults, GetSCCProvider(), State->Changelist);
 	}
 
 	return bUpdatedCachedStates || bUpdatedChangelists;
