Algorithmic optimization in PerforceRevertWorker
* Especially slow when large numbers of files are checked out (in the 25k+ range)
* Existing algorithm scales cubically
* Optimization: Changed to take advantage of the fact that filenames are already in a map, so a cheaper lookup to check if file in changelist is part of the to-be-deleted set
* Benchmark for this change is a map with ~32k empty actors in it.  Source control enabled and using world partition (and thus One File Per Actor).  It is important to check out all __ExternalActors__ files into default changelist and have the "View changes" window open.
- Delete N actors
- Save the level and measure time

Before:
100 actors : ~20s
200 actors : ~1m 20s
300 actors : ~2m 40s
500 actors : ~7m 20s

After:
100 actors : ~6s
200 actors : ~14s
500 actors : ~25s

There is plenty of room for more optimization here.

diff --git a/Engine/Plugins/Developer/PerforceSourceControl/Source/PerforceSourceControl/Private/PerforceSourceControlOperations.cpp b/Engine/Plugins/Developer/PerforceSourceControl/Source/PerforceSourceControl/Private/PerforceSourceControlOperations.cpp
--- a/Engine/Plugins/Developer/PerforceSourceControl/Source/PerforceSourceControl/Private/PerforceSourceControlOperations.cpp
+++ b/Engine/Plugins/Developer/PerforceSourceControl/Source/PerforceSourceControl/Private/PerforceSourceControlOperations.cpp
@@ -566,9 +566,7 @@ static bool RemoveFilesFromChangelist(const TMap<FString, EPerforceState::Type>&
 {
 	return ChangelistState->Files.RemoveAll([&Results](FSourceControlStateRef& State) -> bool
 		{
-			return Algo::AnyOf(Results, [&State](auto& Result) {
-				return State->GetFilename() == Result.Key;
-				});
+			return Results.Contains(State->GetFilename());
 		}) > 0;
 }
 
