Running FGetFile while the perforce source control provider is active will no longer incorrectly mark the connection as dropped, preventing future commands from being run.

#rnx

### Bug Fix
- In CL 24941987 FPerforceGetFileWorker::Execute was modified to remove the perforce connection being created as that connection was not used and FPerforceSourceControlRevision::Get would create a second connection to be used.
- Removing this connection has an unintended knock on effect in that the command passed to the worker was never marked as having a successful server connection. So next tick the source control provider would check the now completed command, see that the connection was not a success (even though the command returned a success) and mark the connection to the server as broken, preventing future commands from working.
- Over the last year or so we have tried to move the code away from these sort of non-obvious side effects but there is a lot of work to do still.
- For now it is safest to add the connection back to FPerforceGetFileWorker::Execute but also add an overload to FPerforceSourceControlRevision::Get that allows us to pass in the connection to be used, to avoid the original problem of creating two connections as connection creation is super slow and can be costly on the server.

### Additional Changes
- There was no reason to allocate the FPerforceSourceControlRevision in FPerforceGetFileWorker::Execute on the heap as it is used immediately then deleted. We might as well just create it on the stack.
- Changed some code logic to early out rather than having long conditional scopes.
- The input parameters for FGetFile are still a bit odd (no revision == #0, revision -1 == #head etc) but leaving that alone for now until we merge FGetFile and FDownloadFiles to eliminate the duplicate code functionality.

diff --git a/Engine/Plugins/Developer/PerforceSourceControl/Source/PerforceSourceControl/Private/PerforceSourceControlOperations.cpp b/Engine/Plugins/Developer/PerforceSourceControl/Source/PerforceSourceControl/Private/PerforceSourceControlOperations.cpp
--- a/Engine/Plugins/Developer/PerforceSourceControl/Source/PerforceSourceControl/Private/PerforceSourceControlOperations.cpp
+++ b/Engine/Plugins/Developer/PerforceSourceControl/Source/PerforceSourceControl/Private/PerforceSourceControlOperations.cpp
@@ -3415,24 +3415,28 @@ bool FPerforceGetFileWorker::Execute(FPerforceSourceControlCommand& InCommand)
 {
 	TRACE_CPUPROFILER_EVENT_SCOPE(FPerforceGetChangelistDetailsWorker::FPerforceGetFileWorker);
 
-	if (!InCommand.IsCanceled())
+	FScopedPerforceConnection ScopedConnection(InCommand);
+
+	if (InCommand.IsCanceled() || !ScopedConnection.IsValid())
 	{
-		TSharedRef<FGetFile, ESPMode::ThreadSafe> Operation = StaticCastSharedRef<FGetFile>(InCommand.Operation);
+		return false;
+	}
 
-		TSharedRef<FPerforceSourceControlRevision, ESPMode::ThreadSafe> Revision = MakeShareable(new FPerforceSourceControlRevision(GetSCCProvider()));
-		Revision->FileName = Operation->GetDepotFilePath();
-		Revision->ChangelistNumber = FCString::Atoi(*Operation->GetChangelistNumber());
-		Revision->RevisionNumber = FCString::Atoi(*Operation->GetRevisionNumber());
-		Revision->bIsShelve = Operation->IsShelve();
-	
-		FString OutFilename;
+	TSharedRef<FGetFile, ESPMode::ThreadSafe> Operation = StaticCastSharedRef<FGetFile>(InCommand.Operation);
 
-		InCommand.bCommandSuccessful = Revision->Get(OutFilename, InCommand.Concurrency);
+	FPerforceSourceControlRevision Revision(GetSCCProvider());
 
-		if (InCommand.bCommandSuccessful)
-		{
-			Operation->SetOutPackageFilename(OutFilename);
-		}
+	Revision.FileName = Operation->GetDepotFilePath();
+	Revision.ChangelistNumber = FCString::Atoi(*Operation->GetChangelistNumber());
+	Revision.RevisionNumber = FCString::Atoi(*Operation->GetRevisionNumber());
+	Revision.bIsShelve = Operation->IsShelve();
+
+	FString OutFilename;
+	InCommand.bCommandSuccessful = Revision.Get(OutFilename, ScopedConnection.GetConnection());
+
+	if (InCommand.bCommandSuccessful)
+	{
+		Operation->SetOutPackageFilename(OutFilename);
 	}
 
 	return InCommand.bCommandSuccessful;
diff --git a/Engine/Plugins/Developer/PerforceSourceControl/Source/PerforceSourceControl/Private/PerforceSourceControlRevision.cpp b/Engine/Plugins/Developer/PerforceSourceControl/Source/PerforceSourceControl/Private/PerforceSourceControlRevision.cpp
--- a/Engine/Plugins/Developer/PerforceSourceControl/Source/PerforceSourceControl/Private/PerforceSourceControlRevision.cpp
+++ b/Engine/Plugins/Developer/PerforceSourceControl/Source/PerforceSourceControl/Private/PerforceSourceControlRevision.cpp
@@ -20,69 +20,83 @@ FPerforceSourceControlRevision:: FPerforceSourceControlRevision(FPerforceSourceC
 {
 }
 
-bool FPerforceSourceControlRevision::Get( FString& InOutFilename, EConcurrency::Type InConcurrency) const
+bool FPerforceSourceControlRevision::Get(FString& InOutFilename, EConcurrency::Type InConcurrency) const
 {
-	bool bCommandOK = false;
-	
 	FScopedPerforceConnection ScopedConnection(InConcurrency, GetSCCProvider());
-	if(ScopedConnection.IsValid())
+	if (ScopedConnection.IsValid())
 	{
 		FPerforceConnection& Connection = ScopedConnection.GetConnection();
-		FP4RecordSet Records;
-		bool bConnectionDropped;
-		TArray<FText> ErrorMessages;
-		TArray<FString> Parameters;
 
-		// Suppress the one-line file header normally added by Perforce.
-		Parameters.Add(TEXT("-q"));
+		return Get(InOutFilename, Connection);
+	}
+	else
+	{
+		return false;
+	}
+}
 
-		// Make temp filename to 'print' to
-		FString RevString;
-		FString RevParam;
-		
-		if (bIsShelve)
-		{
-			RevString = FString::Printf(TEXT("%d"), ChangelistNumber);
-			RevParam = TEXT("@=") + RevString;
-		}
-		else
-		{
-			RevString = (RevisionNumber < 0) ? TEXT("head") : FString::Printf(TEXT("%d"), RevisionNumber);
-			RevParam = TEXT("#") + RevString;
-		}
-		
-		FString AbsoluteFileName;
-		if(InOutFilename.Len() > 0)
-		{
-			AbsoluteFileName = InOutFilename;
-		}
-		else
-		{
-			const FString File = FString::Printf(TEXT("%s-Rev-%s-"), *FPaths::GetBaseFilename(FileName), *RevString);
-			const FString Extension = TEXT(".") + FPaths::GetExtension(FileName);
-			const FString TempFileName = FPaths::CreateTempFilename(*FPaths::DiffDir(), *File, *Extension);
-			AbsoluteFileName = FPaths::ConvertRelativePathToFull(TempFileName);
-		}
+bool FPerforceSourceControlRevision::Get(FString& InOutFilename, FPerforceConnection& Connection) const
+{
+	if (!Connection.IsValidConnection())
+	{
+		return false;
+	}
 
-		// output to file
-		Parameters.Add(FString("-o") + AbsoluteFileName);
-		Parameters.Add(FileName + RevParam);
-		
-		bCommandOK = Connection.RunCommand(TEXT("print"), Parameters, Records, ErrorMessages, FOnIsCancelled(), bConnectionDropped);
-		if(bCommandOK)
-		{
-			InOutFilename = AbsoluteFileName;
-		}
-		else
+	FP4RecordSet Records;
+	bool bConnectionDropped;
+	TArray<FText> ErrorMessages;
+	TArray<FString> Parameters;
+
+	// Suppress the one-line file header normally added by Perforce.
+	Parameters.Add(TEXT("-q"));
+
+	// Make temp filename to 'print' to
+	FString RevString;
+	FString RevParam;
+
+	if (bIsShelve)
+	{
+		RevString = FString::Printf(TEXT("%d"), ChangelistNumber);
+		RevParam = TEXT("@=") + RevString;
+	}
+	else
+	{
+		RevString = (RevisionNumber < 0) ? TEXT("head") : FString::Printf(TEXT("%d"), RevisionNumber);
+		RevParam = TEXT("#") + RevString;
+	}
+
+	FString AbsoluteFileName;
+	if (InOutFilename.Len() > 0)
+	{
+		AbsoluteFileName = InOutFilename;
+	}
+	else
+	{
+		const FString File = FString::Printf(TEXT("%s-Rev-%s-"), *FPaths::GetBaseFilename(FileName), *RevString);
+		const FString Extension = TEXT(".") + FPaths::GetExtension(FileName);
+		const FString TempFileName = FPaths::CreateTempFilename(*FPaths::DiffDir(), *File, *Extension);
+		AbsoluteFileName = FPaths::ConvertRelativePathToFull(TempFileName);
+	}
+
+	// output to file
+	Parameters.Add(FString("-o") + AbsoluteFileName);
+	Parameters.Add(FileName + RevParam);
+
+	const bool bCmdResult = Connection.RunCommand(TEXT("print"), Parameters, Records, ErrorMessages, FOnIsCancelled(), bConnectionDropped);
+	if (bCmdResult)
+	{
+		InOutFilename = AbsoluteFileName;
+		return true;
+	}
+	else
+	{
+		for (auto Iter(ErrorMessages.CreateConstIterator()); Iter; Iter++)
 		{
-			for(auto Iter(ErrorMessages.CreateConstIterator()); Iter; Iter++)
-			{
-				FMessageLog("SourceControl").Error(FText::Format(LOCTEXT("PerforceSourceControlRevisionGetErrorFormat", "FPerforceSourceControlRevision::Get print Error: {0}"), *Iter));
-			}
+			FMessageLog("SourceControl").Error(FText::Format(LOCTEXT("PerforceSourceControlRevisionGetErrorFormat", "FPerforceSourceControlRevision::Get print Error: {0}"), *Iter));
 		}
-	}
 
-	return bCommandOK;
+		return false;
+	}
 }
 
 /** 
diff --git a/Engine/Plugins/Developer/PerforceSourceControl/Source/PerforceSourceControl/Private/PerforceSourceControlRevision.h b/Engine/Plugins/Developer/PerforceSourceControl/Source/PerforceSourceControl/Private/PerforceSourceControlRevision.h
--- a/Engine/Plugins/Developer/PerforceSourceControl/Source/PerforceSourceControl/Private/PerforceSourceControlRevision.h
+++ b/Engine/Plugins/Developer/PerforceSourceControl/Source/PerforceSourceControl/Private/PerforceSourceControlRevision.h
@@ -6,15 +6,18 @@
 #include "ISourceControlProvider.h"
 #include "ISourceControlRevision.h"
 
+class FPerforceConnection;
 class FPerforceSourceControlProvider;
 
-class FPerforceSourceControlRevision : public ISourceControlRevision, public TSharedFromThis<FPerforceSourceControlRevision, ESPMode::ThreadSafe>
+class FPerforceSourceControlRevision final : public ISourceControlRevision, public TSharedFromThis<FPerforceSourceControlRevision, ESPMode::ThreadSafe>
 {
 public:
 
 	FPerforceSourceControlRevision(FPerforceSourceControlProvider& InSCCProvider);
 	virtual ~FPerforceSourceControlRevision() = default;
 
+	bool Get(FString& InOutFilename, FPerforceConnection& InConnection) const;
+
 	/** ISourceControlRevision interface */
 	virtual bool Get( FString& InOutFilename, EConcurrency::Type InConcurrency = EConcurrency::Synchronous) const override;
 	virtual bool GetAnnotated( TArray<FAnnotationLine>& OutLines ) const override;
