Don't save empty perforce setting keys to the config file, remove them instead to make the file more readable.

#rnx

diff --git a/Engine/Plugins/Developer/PerforceSourceControl/Source/PerforceSourceControl/Private/PerforceSourceControlSettings.cpp b/Engine/Plugins/Developer/PerforceSourceControl/Source/PerforceSourceControl/Private/PerforceSourceControlSettings.cpp
--- a/Engine/Plugins/Developer/PerforceSourceControl/Source/PerforceSourceControl/Private/PerforceSourceControlSettings.cpp
+++ b/Engine/Plugins/Developer/PerforceSourceControl/Source/PerforceSourceControl/Private/PerforceSourceControlSettings.cpp
@@ -158,10 +158,23 @@ void FPerforceSourceControlSettings::SaveSettings() const
 	const FString& IniFile = SourceControlHelpers::GetSettingsIni();
 	
 	GConfig->SetBool(*SettingsSection, TEXT("UseP4Config"), ConnectionInfo.bUseP4Config, IniFile);
-	GConfig->SetString(*SettingsSection, TEXT("Port"), *ConnectionInfo.Port, IniFile);
-	GConfig->SetString(*SettingsSection, TEXT("UserName"), *ConnectionInfo.UserName, IniFile);
-	GConfig->SetString(*SettingsSection, TEXT("Workspace"), *ConnectionInfo.Workspace, IniFile);
-	GConfig->SetString(*SettingsSection, TEXT("HostOverride"), *ConnectionInfo.HostOverride, IniFile);
+	
+	auto AddOrRemoveKey = [this, &IniFile](const TCHAR* Key, const FString& Value)
+		{
+			if (!Value.IsEmpty())
+			{
+				GConfig->SetString(*SettingsSection, Key, *Value, IniFile);
+			}
+			else
+			{
+				GConfig->RemoveKey(*SettingsSection, Key, IniFile);
+			}
+		};
+
+	AddOrRemoveKey(TEXT("Port"), ConnectionInfo.Port);
+	AddOrRemoveKey(TEXT("UserName"), ConnectionInfo.UserName);
+	AddOrRemoveKey(TEXT("Workspace"), ConnectionInfo.Workspace);
+	AddOrRemoveKey(TEXT("HostOverride"), ConnectionInfo.HostOverride);
 }
 
 FPerforceConnectionInfo FPerforceSourceControlSettings::GetConnectionInfo() const
