Optimize FPerforceGetPendingChangelistsWorker when getting files for all changelists by executing a single 'opened' command and parsing the results.
Speeds up query from 20+s to around 1.5s in an example workspace with 54 changelists.

diff --git a/Engine/Plugins/Developer/PerforceSourceControl/Source/PerforceSourceControl/Private/PerforceSourceControlOperations.cpp b/Engine/Plugins/Developer/PerforceSourceControl/Source/PerforceSourceControl/Private/PerforceSourceControlOperations.cpp
--- a/Engine/Plugins/Developer/PerforceSourceControl/Source/PerforceSourceControl/Private/PerforceSourceControlOperations.cpp
+++ b/Engine/Plugins/Developer/PerforceSourceControl/Source/PerforceSourceControl/Private/PerforceSourceControlOperations.cpp
@@ -2057,13 +2057,27 @@ FName FPerforceGetPendingChangelistsWorker::GetName() const
 	return "GetPendingChangelists";
 }
 
-static bool GetOpenedFilesInChangelist(FPerforceConnection& Connection, FPerforceSourceControlCommand& InCommand, const FPerforceSourceControlChangelist& Changelist, TArray<FPerforceSourceControlState>& FilesStates)
+static bool GetOpenedFilesInChangelist(
+	FPerforceConnection& Connection,
+	FPerforceSourceControlCommand& InCommand,
+	TOptional<FPerforceSourceControlChangelist> Changelist,
+	TArray<FPerforceSourceControlState>& FilesStates)
 {
 	TRACE_CPUPROFILER_EVENT_SCOPE(FPerforce::GetOpenedFilesInChangelist);
 
 	TArray<FString> Parameters;
-	Parameters.Add(TEXT("-c"));	// -c	Changelist
-	Parameters.Add(Changelist.ToString());	// <changelist>
+	if (Changelist.IsSet())
+	{
+		Parameters.Add(TEXT("-c"));	// -c	Changelist
+		Parameters.Add(Changelist.GetValue().ToString());	// <changelist>
+	}
+	else 
+	{
+		Parameters.Add(TEXT("-u"));									// -u			For user
+		Parameters.Add(InCommand.ConnectionInfo.UserName);			// <username>
+		Parameters.Add(TEXT("-C"));									// -C			For workspace
+		Parameters.Add(InCommand.ConnectionInfo.Workspace);			// <workspace>
+	}
 
 	FP4RecordSet Records;
 	Connection.RunCommand(TEXT("opened"), Parameters, Records, InCommand.ResultInfo.ErrorMessages, FOnIsCancelled::CreateRaw(&InCommand, &FPerforceSourceControlCommand::IsCanceled), InCommand.bConnectionDropped);
@@ -2180,15 +2194,30 @@ bool FPerforceGetPendingChangelistsWorker::Execute(FPerforceSourceControlCommand
 		if (Operation->ShouldUpdateFilesStates())
 		{
 			OutCLFilesStates.Reserve(OutChangelistsStates.Num());
+			TArray<FPerforceSourceControlState> AllFilesStates;
+			GetOpenedFilesInChangelist(Connection, InCommand, {}, AllFilesStates);
 
-			for (FPerforceSourceControlChangelistState& ChangelistState : OutChangelistsStates)
+			if (ShouldContinueProcessing() && AllFilesStates.Num() && OutChangelistsStates.Num())
 			{
-				if (!ShouldContinueProcessing())
+				Algo::SortBy(AllFilesStates, [](const FPerforceSourceControlState& State){ return State.Changelist.ToInt(); });
+				Algo::SortBy(OutChangelistsStates, [](const FPerforceSourceControlChangelistState& State){ 
+					return StaticCastSharedRef<FPerforceSourceControlChangelist>(State.GetChangelist())->ToInt();});
+				
+				OutCLFilesStates.AddDefaulted(OutChangelistsStates.Num());
+				for (int32 ChangelistIndex = 0, FileIndex = 0; FileIndex < AllFilesStates.Num() && ChangelistIndex < OutChangelistsStates.Num();)
 				{
-					break;
+					int32 FileChangelist = AllFilesStates[FileIndex].Changelist.ToInt();
+					int32 ChangelistChangelist = StaticCastSharedRef<FPerforceSourceControlChangelist>(OutChangelistsStates[ChangelistIndex].GetChangelist())->ToInt();
+					if(FileChangelist == ChangelistChangelist)
+					{
+						OutCLFilesStates[ChangelistIndex].Add(MoveTemp(AllFilesStates[FileIndex]));
+						++FileIndex;
+					}
+					else
+					{
+						++ChangelistIndex;
+					}
 				}
-
-				GetOpenedFilesInChangelist(Connection, InCommand, ChangelistState.Changelist, OutCLFilesStates.Emplace_GetRef());
 			}
 		}
 
