Truncation warning fixes for PerforceSourceControl module

#rnx

diff --git a/Engine/Plugins/Developer/PerforceSourceControl/Source/PerforceSourceControl/PerforceSourceControl.Build.cs b/Engine/Plugins/Developer/PerforceSourceControl/Source/PerforceSourceControl/PerforceSourceControl.Build.cs
--- a/Engine/Plugins/Developer/PerforceSourceControl/Source/PerforceSourceControl/PerforceSourceControl.Build.cs
+++ b/Engine/Plugins/Developer/PerforceSourceControl/Source/PerforceSourceControl/PerforceSourceControl.Build.cs
@@ -8,8 +8,8 @@ public class PerforceSourceControl : ModuleRules
 	{
 		IWYUSupport = IWYUSupport.KeepAsIsForNow;
 
-        PrivateDependencyModuleNames.AddRange(
-            new string[] {
+		PrivateDependencyModuleNames.AddRange(
+			new string[] {
 				"Core",
 				"InputCore",
 				"Slate",
@@ -23,7 +23,9 @@ public class PerforceSourceControl : ModuleRules
 		if (Target.Platform == UnrealTargetPlatform.Win64 || Target.Platform == UnrealTargetPlatform.Mac)
 		{
 			AddEngineThirdPartyPrivateStaticDependencies(Target, "OpenSSL");
-            AddEngineThirdPartyPrivateStaticDependencies(Target, "zlib");
+			AddEngineThirdPartyPrivateStaticDependencies(Target, "zlib");
 		}
+
+		UnsafeTypeCastWarningLevel = WarningLevel.Error;
 	}
 }
diff --git a/Engine/Plugins/Developer/PerforceSourceControl/Source/PerforceSourceControl/Private/SPerforceSourceControlSettings.cpp b/Engine/Plugins/Developer/PerforceSourceControl/Source/PerforceSourceControl/Private/SPerforceSourceControlSettings.cpp
--- a/Engine/Plugins/Developer/PerforceSourceControl/Source/PerforceSourceControl/Private/SPerforceSourceControlSettings.cpp
+++ b/Engine/Plugins/Developer/PerforceSourceControl/Source/PerforceSourceControl/Private/SPerforceSourceControlSettings.cpp
@@ -179,7 +179,7 @@ void SPerforceSourceControlSettings::Construct(const FArguments& InArgs, FPerfor
 			[
 				SAssignNew(WorkspaceCombo, SComboButton)
 				.OnGetMenuContent(this, &SPerforceSourceControlSettings::OnGetMenuContent)
-				.ContentPadding(1)
+				.ContentPadding(1.0f)
 				.ToolTipText( LOCTEXT("AutoWorkspaces_Tooltip", "Choose from a list of available workspaces. Requires a server and username before use.") )
 				.IsEnabled(this, &SPerforceSourceControlSettings::IsP4ConfigDisabled)
 				.ButtonContent()
