Fixing partial targeted CL updates from perforce not being updated with their filestates due to an assumption that all CLs were always requested

diff --git a/Engine/Plugins/Developer/PerforceSourceControl/Source/PerforceSourceControl/Private/PerforceSourceControlOperations.cpp b/Engine/Plugins/Developer/PerforceSourceControl/Source/PerforceSourceControl/Private/PerforceSourceControlOperations.cpp
--- a/Engine/Plugins/Developer/PerforceSourceControl/Source/PerforceSourceControl/Private/PerforceSourceControlOperations.cpp
+++ b/Engine/Plugins/Developer/PerforceSourceControl/Source/PerforceSourceControl/Private/PerforceSourceControlOperations.cpp
@@ -2215,7 +2215,14 @@ bool FPerforceGetPendingChangelistsWorker::Execute(FPerforceSourceControlCommand
 					}
 					else
 					{
-						++ChangelistIndex;
+						if (ChangelistChangelist < FileChangelist)
+						{
+							++ChangelistIndex;
+						}
+						else
+						{
+							++FileIndex;
+						}
 					}
 				}
 			}
