Alternative fix for mismatched changelist and file arrays in update operation.
Rewrite loop to be clear that the primary cursor is over changelists and the secondary cursor is over files.

diff --git a/Engine/Plugins/Developer/PerforceSourceControl/Source/PerforceSourceControl/Private/PerforceSourceControlOperations.cpp b/Engine/Plugins/Developer/PerforceSourceControl/Source/PerforceSourceControl/Private/PerforceSourceControlOperations.cpp
--- a/Engine/Plugins/Developer/PerforceSourceControl/Source/PerforceSourceControl/Private/PerforceSourceControlOperations.cpp
+++ b/Engine/Plugins/Developer/PerforceSourceControl/Source/PerforceSourceControl/Private/PerforceSourceControlOperations.cpp
@@ -2203,25 +2203,26 @@ bool FPerforceGetPendingChangelistsWorker::Execute(FPerforceSourceControlCommand
 				Algo::SortBy(OutChangelistsStates, [](const FPerforceSourceControlChangelistState& State){ 
 					return StaticCastSharedRef<FPerforceSourceControlChangelist>(State.GetChangelist())->ToInt();});
 				
+				// OutChangelistStates is the authoritative list of changelists we want to provide results for, any file states which don't have a matching changelist can be discarded
 				OutCLFilesStates.AddDefaulted(OutChangelistsStates.Num());
-				for (int32 ChangelistIndex = 0, FileIndex = 0; FileIndex < AllFilesStates.Num() && ChangelistIndex < OutChangelistsStates.Num();)
+				for (int32 ChangelistIndex = 0, FileIndex = 0; FileIndex < AllFilesStates.Num() && ChangelistIndex < OutChangelistsStates.Num(); ++ChangelistIndex)
 				{
-					int32 FileChangelist = AllFilesStates[FileIndex].Changelist.ToInt();
-					int32 ChangelistChangelist = StaticCastSharedRef<FPerforceSourceControlChangelist>(OutChangelistsStates[ChangelistIndex].GetChangelist())->ToInt();
-					if(FileChangelist == ChangelistChangelist)
+					const int32 ChangelistChangelist = StaticCastSharedRef<FPerforceSourceControlChangelist>(OutChangelistsStates[ChangelistIndex].GetChangelist())->ToInt();
+					for (; FileIndex < AllFilesStates.Num(); ++FileIndex)
 					{
-						OutCLFilesStates[ChangelistIndex].Add(MoveTemp(AllFilesStates[FileIndex]));
-						++FileIndex;
-					}
-					else
-					{
-						if (ChangelistChangelist < FileChangelist)
+						int32 FileChangelist = AllFilesStates[FileIndex].Changelist.ToInt();
+						if(FileChangelist == ChangelistChangelist)
+						{
+							OutCLFilesStates[ChangelistIndex].Add(MoveTemp(AllFilesStates[FileIndex]));
+						}
+						else if (FileChangelist < ChangelistChangelist)
 						{
-							++ChangelistIndex;
+							// Skip files in skipped changelists until we find a file for this changelist or a later one
+							continue;
 						}
-						else
+						else 
 						{
-							++FileIndex;
+							break;
 						}
 					}
 				}
