Add a new initialization path to ISourceControlProvider allowing the caller to customize the initialization logic and retrieve more info about the initialization itself.

### ISourceControlProvider
- Add a new ISourceControlProvider::Init overload that returns a struct (FInitResult) and takes a bitfield of flags (EInitFlags) so that we can more easily extend the initialization in future and return things to the caller.
- If a provider does not implement the new overload we fall-back to the original. This will be removed once all providers support the new path.

### PerforceSourceControlProvider
- Support the new init options via FPerforceConnection::EnsureValidConnection and update places that call it.

### FPerforceConnection
- Update ::AutoDetectWorkspace to return errors to the caller rather than handle its own logging.
- ::TestLoginConnection and ::TestClientConnection no longer calls FPerforceSourceControlProvider::SetLastErrors, this will be done at a higher level.
- ::AutoDetectWorkspace still does its own logging which will be fixed in a future update, it does however return its errors to the caller.
- Split the implementation of ::EnsureValidConnection to a stand alone internal function, the original function now handles and consolidates error handling.
-- Changed the logic to early out of failure rather than constantly checking a boolean all the way down.
-- Tried to remove the need for FPerforceSourceControlProvider to be passed in but a few places still require it.
-- Fixed our handling of the P4PORT value so that we can at least report problems if the address uses unicode characters. I have not yet tested if we can actually connect to such an address. We no longer use the servers unicode setting to decide the port strings encoding as it is not affected by the server settings.
-- Each error encountered should now record a localized error that we expect the user to be able to react to, then an array of errors from the API itself.
-- Once the internal code has run we will do a pass on all errors (if any) and strip any whitespace at the end (p4 tends to end error messages with newline characters) and then add a final entry to the additional errors detailing the settings at the point of failure. This was mostly being done before but is now enforced for every error.
-- No longer log the ticket, although the workflow to allow the user to supply one is not common we still probably don't want to save that to log files.

diff --git a/Engine/Plugins/Developer/PerforceSourceControl/Source/PerforceSourceControl/Private/PerforceConnection.cpp b/Engine/Plugins/Developer/PerforceSourceControl/Source/PerforceSourceControl/Private/PerforceConnection.cpp
--- a/Engine/Plugins/Developer/PerforceSourceControl/Source/PerforceSourceControl/Private/PerforceConnection.cpp
+++ b/Engine/Plugins/Developer/PerforceSourceControl/Source/PerforceSourceControl/Private/PerforceConnection.cpp
@@ -524,7 +524,7 @@ public:
 	FOnIsCancelled IsCancelled;
 };
 
-static bool TestLoginConnection(FPerforceSourceControlProvider& SCCProvider, ClientApi& P4Client, bool bIsUnicodeServer, TArray<FText>& OutErrorMessages)
+static bool TestLoginConnection(ClientApi& P4Client, bool bIsUnicodeServer, TArray<FText>& OutErrorMessages)
 {
 	TRACE_CPUPROFILER_EVENT_SCOPE(FPerforce::TestLoginConnection);
 
@@ -537,8 +537,6 @@ static bool TestLoginConnection(FPerforceSourceControlProvider& SCCProvider, Cli
 	P4Client.SetArgv(1, const_cast<char* const*>(ArgV));
 	P4Client.Run("login", &User);
 
-	SCCProvider.SetLastErrors(OutErrorMessages);
-
 	return OutErrorMessages.IsEmpty();
 }
 
@@ -546,7 +544,7 @@ static bool TestLoginConnection(FPerforceSourceControlProvider& SCCProvider, Cli
  * Runs "client" command to test if the connection is actually OK. ClientApi::Init() only checks
  * if it can connect to server, doesn't verify user name nor workspace name.
  */
-static bool TestClientConnection(FPerforceSourceControlProvider& SCCProvider, ClientApi& P4Client, const FString& ClientSpecName, bool bIsUnicodeServer, TArray<FText>& OutErrorMessages)
+static bool TestClientConnection(ClientApi& P4Client, const FString& ClientSpecName, bool bIsUnicodeServer, TArray<FText>& OutErrorMessages)
 {
 	TRACE_CPUPROFILER_EVENT_SCOPE(FPerforce::TestClientConnection);
 
@@ -575,8 +573,6 @@ static bool TestClientConnection(FPerforceSourceControlProvider& SCCProvider, Cl
 	// clean up args
 	delete [] ClientSpecUTF8Name;
 
-	SCCProvider.SetLastErrors(OutErrorMessages);
-
 	// If there are error messages, user name is most likely invalid. Otherwise, make sure workspace actually
 	// exists on server by checking if we have it's update date.
 	bool bConnectionOK = OutErrorMessages.Num() == 0 && Records.Num() > 0 && Records[0].Contains(TEXT("Update"));
@@ -607,6 +603,49 @@ static bool CheckUnicodeStatus(ClientApi& P4Client, bool& bIsUnicodeServer, TArr
 	return OutErrorMessages.Num() == 0;
 }
 
+static void FinalizeErrors(const FPerforceSourceControlProvider& SCCProvider, const FPerforceConnectionInfo& Settings, ISourceControlProvider::FInitResult::FConnectionErrors& OutConnectionErrors)
+{
+	if (OutConnectionErrors.AdditionalErrors.IsEmpty())
+	{
+		OutConnectionErrors.AdditionalErrors.Add(LOCTEXT("P4_UnknownError", "Unknown error"));
+	}
+	else
+	{
+		// Some perforce errors end with newline characters which can cause odd formatting
+		// when we display/log the errors, so take this opportunity to trim them if needed.
+		for (FText& MsgText : OutConnectionErrors.AdditionalErrors)
+		{
+			const FString Msg = MsgText.ToString();
+			if (!Msg.IsEmpty() && FChar::IsWhitespace(Msg[Msg.Len() - 1]))
+			{
+				MsgText = FText::FromString(Msg.TrimEnd());
+			}	
+		}
+	}
+
+	// Not really an error message but when we do have errors it is useful to print out a summary of the current settings.
+
+	TArray<FText> Components;
+	Components.Add(FText::Format(LOCTEXT("P4_OwningSystem", "OwningSystem={0}"), FText::FromString(SCCProvider.GetOwnerName())));
+
+	if (!Settings.Port.IsEmpty())
+	{
+		Components.Add(FText::Format(LOCTEXT("P4_Port", "Port={0}"), FText::FromString(Settings.Port)));
+	}
+
+	if (!Settings.UserName.IsEmpty())
+	{
+		Components.Add(FText::Format(LOCTEXT("P4_User", "User={0}"), FText::FromString(Settings.UserName)));
+	}
+
+	if (!Settings.Workspace.IsEmpty())
+	{
+		Components.Add(FText::Format(LOCTEXT("P4_ClientSpec", "ClientSpec={0}"), FText::FromString(Settings.Workspace)));
+	}
+
+	OutConnectionErrors.AdditionalErrors.Add(FText::Join(LOCTEXT("P4_Delim", ", "), Components));
+}
+
 FPerforceConnection::FPerforceConnection(const FPerforceConnectionInfo& InConnectionInfo, FPerforceSourceControlProvider& InSCCProvider)
 	: bEstablishedConnection(false)
 	, bIsUnicode(false)
@@ -621,28 +660,35 @@ FPerforceConnection::~FPerforceConnection()
 	Disconnect();
 }
 
-bool FPerforceConnection::AutoDetectWorkspace(const FPerforceConnectionInfo& InConnectionInfo, FPerforceSourceControlProvider& SCCProvider, FString& OutWorkspaceName)
+bool FPerforceConnection::AutoDetectWorkspace(const FPerforceConnectionInfo& InConnectionInfo, FPerforceSourceControlProvider& SCCProvider, FString& OutWorkspaceName, TArray<FText>& OutErrorMessages)
 {
-	bool Result = false;
-	FMessageLog SourceControlLog("SourceControl");
-
 	//before even trying to summon the window, try to "smart" connect with the default server/username
-	TArray<FText> ErrorMessages;
+
 	FPerforceConnection Connection(InConnectionInfo, SCCProvider);
 	TArray<FString> ClientSpecList;
-	Connection.GetWorkspaceList(InConnectionInfo, FOnIsCancelled(), ClientSpecList, ErrorMessages);
+	Connection.GetWorkspaceList(InConnectionInfo, FOnIsCancelled(), ClientSpecList, OutErrorMessages);
+
+	if (!OutErrorMessages.IsEmpty())
+	{
+		return false;
+	}
 
-	//if only one client spec matched (and default connection info was correct)
 	if (ClientSpecList.Num() == 1)
 	{
 		OutWorkspaceName = ClientSpecList[0];
+
+		FTSMessageLog SourceControlLog("SourceControl");
+
 		FFormatNamedArguments Arguments;
 		Arguments.Add( TEXT("WorkspaceName"), FText::FromString(OutWorkspaceName) );
+
 		SourceControlLog.Info(FText::Format(LOCTEXT("ClientSpecAutoDetect", "Auto-detected Perforce client spec: '{WorkspaceName}'"), Arguments));
-		Result = true;
+		
+		return true;
 	}
 	else if (ClientSpecList.Num() > 0)
 	{
+		FTSMessageLog SourceControlLog("SourceControl");
 		SourceControlLog.Warning(LOCTEXT("AmbiguousClientSpecLine1", "Revision Control unable to auto-login due to ambiguous client specs"));
 		SourceControlLog.Warning(LOCTEXT("AmbiguousClientSpecLine2", "  Please select a client spec in the Perforce settings dialog"));
 		SourceControlLog.Warning(LOCTEXT("AmbiguousClientSpecLine3", "  If you are unable to work with revision control, consider checking out the files by hand temporarily"));
@@ -655,9 +701,17 @@ bool FPerforceConnection::AutoDetectWorkspace(const FPerforceConnectionInfo& InC
 			Arguments.Add( TEXT("ClientSpecName"), FText::FromString(ClientSpecList[Index]) );
 			SourceControlLog.Info(FText::Format(LOCTEXT("AmbiguousClientSpecListItem", "...{ClientSpecName}"), Arguments));
 		}
-	}
 
-	return Result;
+		return false;
+	}
+	else
+	{
+		// No clients
+		FTSMessageLog SourceControlLog("SourceControl");
+		SourceControlLog.Warning(LOCTEXT("NoClientSpec", "Revision Control unable to auto-login as no client specs were found"));
+		
+		return false;
+	}
 }
 
 bool FPerforceConnection::Login(const FPerforceConnectionInfo& InConnectionInfo)
@@ -683,194 +737,157 @@ bool FPerforceConnection::Login(const FPerforceConnectionInfo& InConnectionInfo)
 	return ErrorMessages.Num() == 0;
 }
 
-bool FPerforceConnection::EnsureValidConnection(FString& InOutServerName, FString& InOutUserName, FString& InOutWorkspaceName,
-												const FPerforceConnectionInfo& InConnectionInfo, FPerforceSourceControlProvider& SCCProvider, 
-												EConnectionOptions Options)
+bool EnsureValidConnectionInternal(const FPerforceConnectionInfo& InSettings, FPerforceSourceControlProvider& SCCProvider, EConnectionOptions Options,
+	FPerforceConnectionInfo& OutSettings, ISourceControlProvider::FInitResult::FConnectionErrors& OutConnectionErrors)
 {
-	TRACE_CPUPROFILER_EVENT_SCOPE(FPerforceConnection::EnsureValidConnection);
-
-	bool bIsUnicodeServer = false;
-	bool bConnectionOK = false;
-
-	FTSMessageLog SourceControlLog("SourceControl");
-
-	FString NewServerName = InOutServerName;
-	FString NewUserName = InOutUserName;
-	FString NewClientSpecName = InOutWorkspaceName;
+	const bool bRequireWorkspace = !EnumHasAllFlags(Options, EConnectionOptions::WorkspaceOptional);
 
 	ClientApi TestP4;
 	TestP4.SetProg("UE");
 	TestP4.SetProtocol("tag", "");
 	TestP4.SetProtocol("enableStreams", "");
 
-	if (!NewServerName.IsEmpty())
+	if (!InSettings.Port.IsEmpty())
 	{
-		TestP4.SetPort(TCHAR_TO_ANSI(*NewServerName));
-
-		if (!InConnectionInfo.HostOverride.IsEmpty())
-		{
-			TestP4.SetHost(TCHAR_TO_ANSI(*InConnectionInfo.HostOverride));
-		}
+		TestP4.SetPort(TCHAR_TO_ANSI(*InSettings.Port));
 	}
 
-	// Add easy access to the localized error message if needed
-	auto GetFailedToConnectMessage = []() -> FText
+	if (!InSettings.HostOverride.IsEmpty())
 	{
-		return LOCTEXT("P4ErrorConnection_FailedToConnect", "P4ERROR: Failed to connect to revision control provider.");
-	};
+		TestP4.SetHost(TCHAR_TO_ANSI(*InSettings.HostOverride));
+	}
 
 	Error P4Error;
 	TestP4.Init(&P4Error);
 
-	bConnectionOK = !P4Error.Test();
-	if (!bConnectionOK)
+	const bool bConnectionResult = !P4Error.Test();
+
+	// Assume UTF8 for the Port and potential errors as encoding of the port is not affected by the server settings
+	OutSettings.Port = UTF8_TO_TCHAR(TestP4.GetPort().Text()); // Record the PORT that we attempted to connect to
+
+	if (!bConnectionResult)
 	{
-		//Connection FAILED
+		OutConnectionErrors.ErrorMessage = LOCTEXT("P4_FailedToConnect", "P4ERROR: Failed to connect to revision control server.");
+
 		StrBuf ErrorMessage;
 		P4Error.Fmt(&ErrorMessage);
-		SourceControlLog.Error(GetFailedToConnectMessage());
-		SourceControlLog.Error(FText::FromString(ANSI_TO_TCHAR(ErrorMessage.Text())));
-		FFormatNamedArguments Arguments;
-		Arguments.Add(TEXT("OwningSystem"), FText::FromString(SCCProvider.GetOwnerName()));
-		Arguments.Add( TEXT("PortName"), FText::FromString(NewServerName) );
-		Arguments.Add( TEXT("Ticket"), FText::FromString(InConnectionInfo.Ticket) );
-		SourceControlLog.Error(FText::Format(LOCTEXT("P4ConnectErrorConnection_Details", "OwningSystem={OwningSystem}, Port={PortName}, Ticket={Ticket}"), Arguments));
+
+		OutConnectionErrors.AdditionalErrors.Add(FText::FromString(UTF8_TO_TCHAR(ErrorMessage.Text())));
+		return false;
 	}
 
-	// run an info command to determine unicode status
-	if(bConnectionOK)
+	bool bIsUnicodeServer = false;
+	if (!CheckUnicodeStatus(TestP4, bIsUnicodeServer, OutConnectionErrors.AdditionalErrors))
 	{
-		TArray<FText> ErrorMessages;
+		OutConnectionErrors.ErrorMessage = LOCTEXT("P4ErrorConnection_CouldNotDetermineUnicodeStatus", "P4ERROR: Could not determine server unicode status.");
+		return false;
+	}
 
-		bConnectionOK = CheckUnicodeStatus(TestP4, bIsUnicodeServer, ErrorMessages);
-		if(!bConnectionOK)
-		{
-			SourceControlLog.Error(LOCTEXT("P4ErrorConnection_CouldNotDetermineUnicodeStatus", "P4ERROR: Could not determine server unicode status."));
-			SourceControlLog.Error(ErrorMessages.Num() > 0 ? ErrorMessages[0] : LOCTEXT("P4ErrorConnection_UnknownError", "Unknown error"));
-			FFormatNamedArguments Arguments;
-			Arguments.Add(TEXT("OwningSystem"), FText::FromString(SCCProvider.GetOwnerName()));
-			Arguments.Add( TEXT("PortName"), FText::FromString(NewServerName) );
-			Arguments.Add( TEXT("Ticket"), FText::FromString(InConnectionInfo.Ticket) );
-			SourceControlLog.Error(FText::Format(LOCTEXT("P4UnicodeErrorConnection_Details", "OwningSystem={OwningSystem}, Port={PortName}, Ticket={Ticket}"), Arguments));
-		}
-		else
-		{
-			if(bIsUnicodeServer)
-			{
-				// set translation mode. From here onwards we need to use UTF8 when using text args
-				TestP4.SetTrans(CharSetApi::UTF_8);
-			}
+	if (bIsUnicodeServer)
+	{
+		// set translation mode. From here onwards we need to use UTF8 when using text args
+		TestP4.SetTrans(CharSetApi::UTF_8);
+	}
 
-			// now we have determined unicode status, we can set the values that can be specified in non-ansi characters
-			TestP4.SetCwd(FROM_TCHAR(*FPaths::RootDir(), bIsUnicodeServer));
-			TestP4.SetUser(FROM_TCHAR(*NewUserName, bIsUnicodeServer));
-			TestP4.SetClient(FROM_TCHAR(*NewClientSpecName, bIsUnicodeServer));
-			TestP4.SetPassword(FROM_TCHAR(*InConnectionInfo.Ticket, bIsUnicodeServer));
+	// now we have determined unicode status, we can set the values that can be specified in non-ansi characters
+	TestP4.SetCwd(FROM_TCHAR(*FPaths::RootDir(), bIsUnicodeServer));
+	TestP4.SetUser(FROM_TCHAR(*InSettings.UserName, bIsUnicodeServer));
+	TestP4.SetClient(FROM_TCHAR(*InSettings.Workspace, bIsUnicodeServer));
+	TestP4.SetPassword(FROM_TCHAR(*InSettings.Ticket, bIsUnicodeServer));
 
-		}
-	}
+	const bool LogInResult = TestLoginConnection(TestP4, bIsUnicodeServer, OutConnectionErrors.AdditionalErrors);
+
+	OutSettings.UserName = TO_TCHAR(TestP4.GetUser().Text(), bIsUnicodeServer);
 
-	// Test that we have a valid p4 ticket
-	if (bConnectionOK)
+	if (!LogInResult)
 	{
-		TArray<FText> ErrorMessages;
-		bConnectionOK = TestLoginConnection(SCCProvider, TestP4, bIsUnicodeServer, ErrorMessages);
+		OutConnectionErrors.ErrorMessage = LOCTEXT("P4_InvalidToken", "Unable to log into revision control server");
+		return false;
+	}
 
-		if (!bConnectionOK)
+	// Try to auto detect the client if none were specified and we require one
+	if (bRequireWorkspace)
+	{
+		FString ClientSpecName = InSettings.Workspace;
+		if (ClientSpecName.IsEmpty())
 		{
-			FString ServerName = TO_TCHAR(TestP4.GetPort().Text(), bIsUnicodeServer);
-			FString UserName = TO_TCHAR(TestP4.GetUser().Text(), bIsUnicodeServer);
+			FPerforceConnectionInfo AutoCredentials = InSettings;
+			AutoCredentials.Port = OutSettings.Port;
+			AutoCredentials.UserName = OutSettings.UserName;
 
-			SourceControlLog.Error(GetFailedToConnectMessage());
-			SourceControlLog.Error(ErrorMessages.Num() > 0 ? ErrorMessages[0] : LOCTEXT("P4ErrorConnection_InvalidToken", "Unable to log in"));
-			FFormatNamedArguments Arguments;
-			Arguments.Add(TEXT("OwningSystem"), FText::FromString(SCCProvider.GetOwnerName()));
-			Arguments.Add(TEXT("PortName"), FText::FromString(MoveTemp(ServerName)));
-			Arguments.Add(TEXT("UserName"), FText::FromString(MoveTemp(UserName)));
+			// AutoDetectWorkspace takes care of the error reporting
+			if (!FPerforceConnection::AutoDetectWorkspace(AutoCredentials, SCCProvider, ClientSpecName, OutConnectionErrors.AdditionalErrors))
+			{
+				return false;
+			}
 
-			SourceControlLog.Error(FText::Format(LOCTEXT("P4LoginErrorConnection_Details", "OwningSystem={OwningSystem}, Port={PortName}, User={UserName}"), Arguments));
+			TestP4.SetClient(FROM_TCHAR(*ClientSpecName, bIsUnicodeServer));
 		}
-	}
 
-	const bool bRequireWorkspace = !EnumHasAllFlags(Options, EConnectionOptions::WorkspaceOptional);
-
-	// Try to auto detect the client if none were specified and we require one
-	if (bConnectionOK && bRequireWorkspace && NewClientSpecName.IsEmpty())
-	{
-		FPerforceConnectionInfo AutoCredentials = InConnectionInfo;
-		AutoCredentials.Port = TO_TCHAR(TestP4.GetPort().Text(), bIsUnicodeServer);
-		AutoCredentials.UserName = TO_TCHAR(TestP4.GetUser().Text(), bIsUnicodeServer);
+		// Test that we found a valid client if we require one
+		
+		const bool bValidClientSpec = TestClientConnection(TestP4, ClientSpecName, bIsUnicodeServer, OutConnectionErrors.AdditionalErrors);
 
-		bConnectionOK = FPerforceConnection::AutoDetectWorkspace(AutoCredentials, SCCProvider, NewClientSpecName);
-		if (bConnectionOK)
+		OutSettings.Workspace = TO_TCHAR(TestP4.GetClient().Text(), bIsUnicodeServer);
+		if (!bValidClientSpec)
 		{
-			TestP4.SetClient(FROM_TCHAR(*NewClientSpecName, bIsUnicodeServer));
+			OutConnectionErrors.ErrorMessage = LOCTEXT("P4ErrorConnection_InvalidWorkspace", "Invalid workspace");
+			return false;
 		}
-	}
-
-	// Test that we found a valid client if we require one
-	if (bConnectionOK && bRequireWorkspace)
-	{
-		TArray<FText> ErrorMessages;
 
-		bConnectionOK = TestClientConnection(SCCProvider, TestP4, NewClientSpecName, bIsUnicodeServer, ErrorMessages);
-		
-		if (!bConnectionOK)
+		// If the workspace name autodetected and is the same as the host we should assume that no valid workspace was found.
+		// TODO: Note that in this case the above call to ::TestClientConnection should have already failed so we could consider removing this check
+		if (InSettings.Workspace.IsEmpty() && OutSettings.Workspace == TO_TCHAR(TestP4.GetHost().Text(), bIsUnicodeServer))
 		{
-			SourceControlLog.Error(GetFailedToConnectMessage());
-			SourceControlLog.Error(ErrorMessages.Num() > 0 ? ErrorMessages[0] : LOCTEXT("P4ErrorConnection_InvalidWorkspace", "Invalid workspace"));
-			FFormatNamedArguments Arguments;
-			Arguments.Add(TEXT("OwningSystem"), FText::FromString(SCCProvider.GetOwnerName()));
-			Arguments.Add( TEXT("PortName"), FText::FromString(NewServerName) );
-			Arguments.Add( TEXT("UserName"), FText::FromString(NewUserName) );
-			Arguments.Add( TEXT("ClientSpecName"), FText::FromString(NewClientSpecName) );
-			Arguments.Add( TEXT("Ticket"), FText::FromString(InConnectionInfo.Ticket) );
-			SourceControlLog.Error(FText::Format(LOCTEXT("P4ClientErrorConnection_Details", "OwningSystem={OwningSystem}, Port={PortName}, User={UserName}, ClientSpec={ClientSpecName}, Ticket={Ticket}"), Arguments));
+			OutConnectionErrors.ErrorMessage = LOCTEXT("P4ErrorConnection_MissingWorkspace", "Missing workspace");
+			OutConnectionErrors.AdditionalErrors.Add(LOCTEXT("P4ErrorConnection_NoWorkspaceFound", "No workspace was found for the current user"));
+			return false;
 		}
 	}
 
 	//whether successful or not, disconnect to clean up
 	TestP4.Final(&P4Error);
-	if (bConnectionOK && P4Error.Test())
+	
+	if (P4Error.Test())
 	{
-		//Disconnect FAILED
-		bConnectionOK = false;
+		OutConnectionErrors.ErrorMessage = LOCTEXT("P4_FailedDisconnect", "P4ERROR: Failed to disconnect from revision control server.");
+
 		StrBuf ErrorMessage;
 		P4Error.Fmt(&ErrorMessage);
-		SourceControlLog.Error(LOCTEXT("P4ErrorFailedDisconnect", "P4ERROR: Failed to disconnect from Server."));
-		SourceControlLog.Error(FText::FromString(TO_TCHAR(ErrorMessage.Text(), bIsUnicodeServer)));
-	}
 
-	//if never specified, take the default connection values
-	if (NewServerName.IsEmpty())
-	{
-		NewServerName = TO_TCHAR(TestP4.GetPort().Text(), bIsUnicodeServer);
+		OutConnectionErrors.AdditionalErrors.Add(FText::FromString(TO_TCHAR(ErrorMessage.Text(), bIsUnicodeServer)));
+		return false;
 	}
 
-	if (NewUserName.IsEmpty())
-	{
-		NewUserName = TO_TCHAR(TestP4.GetUser().Text(), bIsUnicodeServer);
-	}
+	return true;
+}
+
+bool FPerforceConnection::EnsureValidConnection(const FPerforceConnectionInfo& InSettings, FPerforceSourceControlProvider& SCCProvider, EConnectionOptions Options,
+	FPerforceConnectionInfo& OutSettings, ISourceControlProvider::FInitResult::FConnectionErrors& OutConnectionErrors)
+{
+	TRACE_CPUPROFILER_EVENT_SCOPE(FPerforceConnection::EnsureValidConnection);
+
+	bool bResult = EnsureValidConnectionInternal(InSettings, SCCProvider, Options, OutSettings, OutConnectionErrors);
 
-	if (NewClientSpecName.IsEmpty() && bRequireWorkspace)
+	if (!bResult)
 	{
-		NewClientSpecName = TO_TCHAR(TestP4.GetClient().Text(), bIsUnicodeServer);
-		if (NewClientSpecName == TO_TCHAR(TestP4.GetHost().Text(), bIsUnicodeServer))
+		FinalizeErrors(SCCProvider, OutSettings, OutConnectionErrors);
+
+		if (!EnumHasAllFlags(Options, EConnectionOptions::SupressErrorLogging))
 		{
-			// If the client spec name is the same as host name, assume P4 couldn't get the actual
-			// spec name for this host and let GetPerforceLogin() try to find a proper one.
-			bConnectionOK = false;
+			FTSMessageLog SourceControlLog("SourceControl");
+
+			SourceControlLog.Error(OutConnectionErrors.ErrorMessage);
+			for (const FText& ErrorMsg : OutConnectionErrors.AdditionalErrors)
+			{
+				SourceControlLog.Error(ErrorMsg);
+			}
 		}
-	}
 
-	if (bConnectionOK)
-	{
-		InOutServerName = NewServerName;
-		InOutUserName = NewUserName;
-		InOutWorkspaceName = NewClientSpecName;
+		SCCProvider.SetLastErrors(OutConnectionErrors.AdditionalErrors);
 	}
 
-	return bConnectionOK;
+	return bResult;
 }
 
 bool FPerforceConnection::GetWorkspaceList(const FPerforceConnectionInfo& InConnectionInfo, FOnIsCancelled InOnIsCanceled, TArray<FString>& OutWorkspaceList, TArray<FText>& OutErrorMessages)
diff --git a/Engine/Plugins/Developer/PerforceSourceControl/Source/PerforceSourceControl/Private/PerforceConnection.h b/Engine/Plugins/Developer/PerforceSourceControl/Source/PerforceSourceControl/Private/PerforceConnection.h
--- a/Engine/Plugins/Developer/PerforceSourceControl/Source/PerforceSourceControl/Private/PerforceConnection.h
+++ b/Engine/Plugins/Developer/PerforceSourceControl/Source/PerforceSourceControl/Private/PerforceConnection.h
@@ -44,8 +44,10 @@ enum class EConnectionOptions : uint8
 {
 	/** No options are applied */
 	None				= 0,
-	/** The connection does not require a workspace to be considered valid*/
+	/** The connection does not require a workspace to be considered valid */
 	WorkspaceOptional	= 1 << 0,
+	/** Errors will not be logged but will still be returned to the caller */
+	SupressErrorLogging	= 1 << 1
 };
 ENUM_CLASS_FLAGS(EConnectionOptions);
 
@@ -72,19 +74,23 @@ public:
 	/** 
 	 * Attempts to automatically detect the workspace to use based on the working directory
 	 */
-	static bool AutoDetectWorkspace(const FPerforceConnectionInfo& InConnectionInfo, FPerforceSourceControlProvider& SCCProvider, FString& OutWorkspaceName);
+	static bool AutoDetectWorkspace(const FPerforceConnectionInfo& InConnectionInfo, FPerforceSourceControlProvider& SCCProvider, FString& OutWorkspaceName, TArray<FText>& OutErrorMessages);
 
 	/**
-	 * Static function in charge of making sure the specified connection is valid or requests that data from the user via dialog
-	 * @param InOutPortName			Port name in the inifile.  Out value is the port name from the connection dialog
-	 * @param InOutUserName			User name in the inifile.  Out value is the user name from the connection dialog
-	 * @param InOutWorkspaceName	Workspace name in the inifile.  Out value is the client spec from the connection dialog
-	 * @param InConnectionInfo		Connection credentials
-	 * @return - true if the connection, whether via dialog or otherwise, is valid.  False if source control should be disabled
+	 * Set up a connection to the server with the given credentials. The function can attempt to autodetect missing credentials or fix incorrect ones with
+	 * the final credentials being returned to the caller.
+	 * 
+	 * @param InSettings			The initial connection credentials.
+	 * @param SCCProvider			The provider that is setting up the connection.
+	 * @param Options				Used to specialize initialization behavior, @see EConnectionOptions.
+	 * @param OutSettings			The finalized connection credentials. If the connection failed then this will contain the credentials that were 
+	 *								used for the step that failed.
+	 * @param OutConnectionErrors	A collection of errors encountered.
+	 * 
+	 * @return - True if a valid connection was established, otherwise false.
 	 */
-	static bool EnsureValidConnection(	FString& InOutServerName, FString& InOutUserName, FString& InOutWorkspaceName,
-										const FPerforceConnectionInfo& InConnectionInfo, FPerforceSourceControlProvider& SCCProvider, 
-										EConnectionOptions Options = EConnectionOptions::None);
+	static bool EnsureValidConnection(	const FPerforceConnectionInfo& InSettings, FPerforceSourceControlProvider& SCCProvider, EConnectionOptions Options,
+										FPerforceConnectionInfo& OutSettings, ISourceControlProvider::FInitResult::FConnectionErrors& OutConnectionErrors);
 
 	/**
 	 * Get List of ClientSpecs
diff --git a/Engine/Plugins/Developer/PerforceSourceControl/Source/PerforceSourceControl/Private/PerforceSourceControlProvider.cpp b/Engine/Plugins/Developer/PerforceSourceControl/Source/PerforceSourceControl/Private/PerforceSourceControlProvider.cpp
--- a/Engine/Plugins/Developer/PerforceSourceControl/Source/PerforceSourceControl/Private/PerforceSourceControlProvider.cpp
+++ b/Engine/Plugins/Developer/PerforceSourceControl/Source/PerforceSourceControl/Private/PerforceSourceControlProvider.cpp
@@ -83,7 +83,17 @@ FPerforceSourceControlProvider::FPerforceSourceControlProvider(const FStringView
 
 void FPerforceSourceControlProvider::Init(bool bForceConnection)
 {
-	ParseCommandLineSettings(bForceConnection);
+	const EInitFlags Flags = bForceConnection ? EInitFlags::AttemptConnection : EInitFlags::None;
+
+	ParseCommandLineSettings(Flags);
+}
+
+ISourceControlProvider::FInitResult FPerforceSourceControlProvider::Init(EInitFlags Flags)
+{
+	FInitResult Result = ParseCommandLineSettings(Flags);
+	Result.bIsAvailable = IsAvailable();
+
+	return Result;
 }
 
 void FPerforceSourceControlProvider::Close()
@@ -205,8 +215,10 @@ bool FPerforceSourceControlProvider::EstablishPersistentConnection()
 	return bIsValidConnection;
 }
 
-void FPerforceSourceControlProvider::ParseCommandLineSettings(bool bForceConnection)
+ISourceControlProvider::FInitResult FPerforceSourceControlProvider::ParseCommandLineSettings(EInitFlags InitFlags)
 {
+	ISourceControlProvider::FInitResult Result;
+
 	FPerforceSourceControlSettings& P4Settings = AccessSettings();
 
 	// First we take a copy of the existing settings
@@ -216,7 +228,7 @@ void FPerforceSourceControlProvider::ParseCommandLineSettings(bool bForceConnect
 	FString HostOverrideName = P4Settings.GetHostOverride();
 	FString Changelist = P4Settings.GetChangelistNumber();
 
-	EConnectionOptions Options = EConnectionOptions::None;
+	EConnectionOptions ConnectionOptions = EConnectionOptions::None;
 
 	// Then we see if any of these settings are overridden by the initial settings
 	// Note that as long as one setting is overridden, we will reset all non-overridden 
@@ -234,7 +246,7 @@ void FPerforceSourceControlProvider::ParseCommandLineSettings(bool bForceConnect
 		// so don't need to automatically find a workspace when ensuring the connection.
 		if (InitialSettings.IsOverridden(TEXT("P4Client")))
 		{
-			Options |= EConnectionOptions::WorkspaceOptional;
+			ConnectionOptions |= EConnectionOptions::WorkspaceOptional;
 		}
 
 		P4Settings.SetPort(PortName);
@@ -244,22 +256,50 @@ void FPerforceSourceControlProvider::ParseCommandLineSettings(bool bForceConnect
 		P4Settings.SetChangelistNumber(Changelist);
 	}
 
-	if (bForceConnection)
+	if (EnumHasAnyFlags(InitFlags, EInitFlags::AttemptConnection))
 	{
 		bLoginError = false;
+
 		FPerforceConnectionInfo ConnectionInfo = P4Settings.GetConnectionInfo();
-		if(FPerforceConnection::EnsureValidConnection(PortName, UserName, ClientSpecName, ConnectionInfo, *this, Options))
+		FPerforceConnectionInfo OutputSettings;
+
+		if (EnumHasAnyFlags(InitFlags, EInitFlags::SupressErrorLogging))
 		{
-			P4Settings.SetPort(PortName);
-			P4Settings.SetUserName(UserName);
-			P4Settings.SetWorkspace(ClientSpecName);
-			P4Settings.SetHostOverride(HostOverrideName);
+			ConnectionOptions |= EConnectionOptions::SupressErrorLogging;
+		}
+
+		if(FPerforceConnection::EnsureValidConnection(ConnectionInfo, *this, ConnectionOptions, OutputSettings, Result.Errors))
+		{
+			// The connection was a success so we should store the values used by the successful connection
+			P4Settings.SetPort(OutputSettings.Port);
+			P4Settings.SetUserName(OutputSettings.UserName);
+			P4Settings.SetWorkspace(OutputSettings.Workspace);
+
 			bServerAvailable = true;
 		}
+
+		// Fill in FInitResult::ConnectionSettings with the actual settings that were used
+
+		if (!OutputSettings.Port.IsEmpty())
+		{
+			Result.ConnectionSettings.Add(ISourceControlProvider::EStatus::Port, OutputSettings.Port);
+		}
+
+		if (!OutputSettings.UserName.IsEmpty())
+		{
+			Result.ConnectionSettings.Add(ISourceControlProvider::EStatus::User, OutputSettings.UserName);
+		}
+
+		if (!OutputSettings.Workspace.IsEmpty())
+		{
+			Result.ConnectionSettings.Add(ISourceControlProvider::EStatus::Client, OutputSettings.Workspace);
+		}
 	}
 
 	//Save off settings so this doesn't happen every time
 	SaveConnectionSettings();
+
+	return Result;
 }
 
 void FPerforceSourceControlProvider::GetWorkspaceList(const FPerforceConnectionInfo& InConnectionInfo, TArray<FString>& OutWorkspaceList, TArray<FText>& OutErrorMessages)
@@ -829,14 +869,17 @@ ECommandResult::Type FPerforceSourceControlProvider::SwitchWorkspace(FStringView
 
 	if (!NewWorkspaceName.IsEmpty())
 	{
-		FString PortName = P4Settings.GetPort();
-		FString UserName = P4Settings.GetUserName();
-		
-		if (FPerforceConnection::EnsureValidConnection(PortName, UserName, WorkspaceName, P4Settings.GetConnectionInfo(), *this, EConnectionOptions::WorkspaceOptional))
+		FPerforceConnectionInfo NewWorkspaceSettings = P4Settings.GetConnectionInfo();
+		NewWorkspaceSettings.Workspace = WorkspaceName;
+
+		ISourceControlProvider::FInitResult Results;
+
+		FPerforceConnectionInfo OutputSettings;
+		if (FPerforceConnection::EnsureValidConnection(NewWorkspaceSettings, *this, EConnectionOptions::WorkspaceOptional, OutputSettings, Results.Errors))
 		{
-			P4Settings.SetPort(PortName);
-			P4Settings.SetUserName(UserName);
-			P4Settings.SetWorkspace(WorkspaceName);
+			P4Settings.SetPort(OutputSettings.Port);
+			P4Settings.SetUserName(OutputSettings.UserName);
+			P4Settings.SetWorkspace(OutputSettings.Workspace);
 
 			bServerAvailable = true;
 
diff --git a/Engine/Plugins/Developer/PerforceSourceControl/Source/PerforceSourceControl/Private/PerforceSourceControlProvider.h b/Engine/Plugins/Developer/PerforceSourceControl/Source/PerforceSourceControl/Private/PerforceSourceControlProvider.h
--- a/Engine/Plugins/Developer/PerforceSourceControl/Source/PerforceSourceControl/Private/PerforceSourceControlProvider.h
+++ b/Engine/Plugins/Developer/PerforceSourceControl/Source/PerforceSourceControl/Private/PerforceSourceControlProvider.h
@@ -24,6 +24,7 @@ public:
 
 	/* ISourceControlProvider implementation */
 	virtual void Init(bool bForceConnection = true) override;
+	virtual FInitResult Init(EInitFlags Flags) override;
 	virtual void Close() override;
 	virtual FText GetStatusText() const override;
 	virtual TMap<EStatus, FString> GetStatus() const override;
@@ -140,9 +141,10 @@ private:
 	void OutputCommandMessages(const class FPerforceSourceControlCommand& InCommand) const;
 
 	/**
-	 * Loads user/SCC information from the INI file.
+	 * Loads user/SCC information from the INI file and can attempt to make a connection to the server if requested by
+	 * the Flags parameter.
 	 */
-	void ParseCommandLineSettings(bool bForceConnection);
+	ISourceControlProvider::FInitResult ParseCommandLineSettings(EInitFlags Flags);
 
 	/**
 	 * Helper function for running command 'synchronously'.
diff --git a/Engine/Source/Developer/SourceControl/Private/ISourceControlProvider.cpp b/Engine/Source/Developer/SourceControl/Private/ISourceControlProvider.cpp
--- a/Engine/Source/Developer/SourceControl/Private/ISourceControlProvider.cpp
+++ b/Engine/Source/Developer/SourceControl/Private/ISourceControlProvider.cpp
@@ -7,6 +7,18 @@
 
 #define LOCTEXT_NAMESPACE "ISourceControlProvider"
 
+ISourceControlProvider::FInitResult ISourceControlProvider::Init(EInitFlags Flags)
+{
+	// TODO: This implementation is for compatibility. Once all providers have their own
+	// implementation it should be removed.
+	const bool bForceConnection = EnumHasAnyFlags(Flags, EInitFlags::AttemptConnection);
+	Init(bForceConnection);
+
+	FInitResult Result;
+	Result.bIsAvailable = IsAvailable();
+
+	return Result;
+}
 
 ECommandResult::Type ISourceControlProvider::Login(const FString& InPassword, EConcurrency::Type InConcurrency, const FSourceControlOperationComplete& InOperationCompleteDelegate)
 {
diff --git a/Engine/Source/Developer/SourceControl/Public/ISourceControlProvider.h b/Engine/Source/Developer/SourceControl/Public/ISourceControlProvider.h
--- a/Engine/Source/Developer/SourceControl/Public/ISourceControlProvider.h
+++ b/Engine/Source/Developer/SourceControl/Public/ISourceControlProvider.h
@@ -87,6 +87,71 @@ DECLARE_MULTICAST_DELEGATE(FSourceControlStateChanged);
  */
 class ISourceControlProvider : public IModularFeature
 {
+public:
+
+	/** Key value for common source control provider statuses */
+	enum class EStatus
+	{
+		Enabled,
+		Connected,
+		Port,
+		User,
+		Client,
+		Repository,
+		Remote,
+		Branch,
+		Email,
+		ScmVersion,
+		PluginVersion,
+		Workspace,
+		WorkspacePath,
+		Changeset
+	};
+
+	/** Flags used when calling ISourceControlProvider::Init */
+	enum class EInitFlags : uint32
+	{
+		/** No flags set */
+		None				= 0,
+		/** A connection attempt to the provider will be made as part of the initialization */
+		AttemptConnection	= 1 << 0,
+		/** Errors will not be logged but will still be returned to the caller */
+		SupressErrorLogging = 1 << 1
+	};
+
+	/** The result info returned form calls to  ISourceControlProvider::Init */
+	struct FInitResult
+	{
+		/** True if the provider is ready for use, otherwise false. This value should match calling ISourceControlProvider::IsAvailable. */
+		bool bIsAvailable = false;
+
+		/** 
+		 * Contains the important settings used when trying to make a server connection.
+		 * This will only contain entries if Init is called with the EInitFlags::AttemptConnection flag.
+		 * 
+		 * This map should still contain entries if the connection failed.
+		 */
+		TMap<EStatus, FString> ConnectionSettings;
+
+		/** Contains errors encountered during initialization */
+		struct FConnectionErrors
+		{
+			/** A higher level error message describing at what point in the connection something went wrong. */
+			FText ErrorMessage;
+
+			/** A collection of lower level error messages usually from the source revision api itself. */
+			TArray<FText> AdditionalErrors;
+
+			FConnectionErrors() = default;
+			FConnectionErrors(const FText& InErrorMessage, const TArray<FText>& InAdditionalErrors)
+				: ErrorMessage(InErrorMessage), AdditionalErrors(InAdditionalErrors)
+			{
+			}
+
+			~FConnectionErrors() = default;
+		} Errors;
+	};
+
 public:
 	/**
 	 * Virtual destructor
@@ -99,6 +164,14 @@ public:
 	 */
 	virtual void Init(bool bForceConnection = true) = 0;
 
+	/**
+	 * Initialize source control provider.
+	 * @param Flags		Used to specialize initialization behavior, @see EInitFlags
+	 * 
+	 * @return			A structure containing the result as well as any errors. @see FInitResult
+	 */
+	SOURCECONTROL_API virtual FInitResult Init(EInitFlags Flags);
+
 	/**
 	 * Shut down source control provider.
 	 */
@@ -110,24 +183,6 @@ public:
 	/** Get the source control status as plain, human-readable text */
 	virtual FText GetStatusText() const = 0;
 
-	enum class EStatus
-	{
-		Enabled,
-		Connected,
-		Port,
-		User,
-		Client,
-		Repository,
-		Remote,
-		Branch,
-		Email,
-		ScmVersion,
-		PluginVersion,
-		Workspace,
-		WorkspacePath,
-		Changeset
-	};
-
 	/** Get the source control status as a series of key value pairs */
 	virtual TMap<EStatus, FString> GetStatus() const = 0;
 
@@ -428,4 +483,4 @@ protected:
 	virtual TUniquePtr<ISourceControlProvider> Create(const FStringView& OwnerName, const FSourceControlInitSettings& InitialSettings) const { return TUniquePtr<ISourceControlProvider>(); }
 };
 
-
+ENUM_CLASS_FLAGS(ISourceControlProvider::EInitFlags);
