Fixed a bug where FPerforceSourceControlProvider::SwitchWorkspace was not correctly switching to the new workspace

#rnx

- Bug was intorduced in CL 26952096
- EConnectionOptions::WorkspaceOptional means we should not test if a workspace is valid but if one is provided it should still be used but in the above change EnsureValidConnection was changed and this resulted in the code discarding the workspace entirely if that option was used.
- If that option is used we just pass the input workspace name (if there is one) to the output structure and trust that the caller knows what they are doing.

diff --git a/Engine/Plugins/Developer/PerforceSourceControl/Source/PerforceSourceControl/Private/PerforceConnection.cpp b/Engine/Plugins/Developer/PerforceSourceControl/Source/PerforceSourceControl/Private/PerforceConnection.cpp
--- a/Engine/Plugins/Developer/PerforceSourceControl/Source/PerforceSourceControl/Private/PerforceConnection.cpp
+++ b/Engine/Plugins/Developer/PerforceSourceControl/Source/PerforceSourceControl/Private/PerforceConnection.cpp
@@ -844,6 +844,10 @@ bool EnsureValidConnectionInternal(const FPerforceConnectionInfo& InSettings, FP
 			return false;
 		}
 	}
+	else
+	{
+		OutSettings.Workspace = InSettings.Workspace;
+	}
 
 	//whether successful or not, disconnect to clean up
 	TestP4.Final(&P4Error);
