Synchronizing source control state from the perforce source control provider into TEDS

diff --git a/Engine/Plugins/Developer/PerforceSourceControl/Source/PerforceSourceControl/PerforceSourceControl.Build.cs b/Engine/Plugins/Developer/PerforceSourceControl/Source/PerforceSourceControl/PerforceSourceControl.Build.cs
--- a/Engine/Plugins/Developer/PerforceSourceControl/Source/PerforceSourceControl/PerforceSourceControl.Build.cs
+++ b/Engine/Plugins/Developer/PerforceSourceControl/Source/PerforceSourceControl/PerforceSourceControl.Build.cs
@@ -16,6 +16,7 @@ public class PerforceSourceControl : ModuleRules
 				"Slate",
 				"SlateCore",
 				"SourceControl",
+				"TypedElementFramework",
 			}
 		);
 
diff --git a/Engine/Plugins/Developer/PerforceSourceControl/Source/PerforceSourceControl/Private/PerforceSourceControlOperations.cpp b/Engine/Plugins/Developer/PerforceSourceControl/Source/PerforceSourceControl/Private/PerforceSourceControlOperations.cpp
--- a/Engine/Plugins/Developer/PerforceSourceControl/Source/PerforceSourceControl/Private/PerforceSourceControlOperations.cpp
+++ b/Engine/Plugins/Developer/PerforceSourceControl/Source/PerforceSourceControl/Private/PerforceSourceControlOperations.cpp
@@ -22,6 +22,11 @@
 #include "PerforceSourceControlRevision.h"
 #include "SourceControlHelpers.h"
 #include "SourceControlOperations.h"
+#include "Elements/Columns/TypedElementPackageColumns.h"
+#include "Elements/Columns/TypedElementRevisionControlColumns.h"
+#include "Elements/Framework/TypedElementIndexHasher.h"
+#include "Elements/Framework/TypedElementRegistry.h"
+#include "Elements/Interfaces/TypedElementDataStorageInterface.h"
 #include "ProfilingDebugging/CpuProfilerTrace.h"
 
 #define LOCTEXT_NAMESPACE "PerforceSourceControl"
@@ -1979,6 +1984,86 @@ bool FPerforceUpdateStatusWorker::UpdateStates() const
 {
 	TRACE_CPUPROFILER_EVENT_SCOPE(FPerforceUpdateStatusWorker::UpdateStates);
 
+	auto UpdateDataStorage = [](const FPerforceSourceControlState& State)
+	{
+		using namespace TypedElementQueryBuilder;
+		using DSI = ITypedElementDataStorageInterface;
+
+		DSI* DataStorage = UTypedElementRegistry::GetInstance()->GetMutableDataStorage();
+		if (!DataStorage)
+		{
+			return;
+		}
+		
+		FString Filename = FPaths::SetExtension(State.GetFilename(), "");
+		FPaths::NormalizeFilename(Filename);
+
+		auto GetRevisionControlRow = [DataStorage](const FString& Filename) -> TypedElementRowHandle
+		{
+			uint64 Index = TypedElementDataStorage::GenerateIndexHash(Filename);
+			TypedElementRowHandle Row = DataStorage->FindIndexedRow(Index);
+
+			if (!DataStorage->IsRowAvailable(Row))
+			{
+				static TypedElementTableHandle Table = DataStorage->FindTable(FName("Editor_RevisionControlTable"));
+				Row = DataStorage->AddRow(Table);
+				DataStorage->IndexRow(Index, Row);
+			}
+			return Row;
+		};
+
+		TypedElementRowHandle Row = GetRevisionControlRow(Filename);
+
+		DataStorage->AddOrGetColumn<FSCCRevisionIdColumn>(Row)->RevisionId.Id[0] = State.LocalRevNumber;
+		DataStorage->AddOrGetColumn<FSCCExternalRevisionIdColumn>(Row)->RevisionId.Id[0] = State.DepotRevNumber;
+
+		TArray<UScriptStruct*> ToAdd;
+		TArray<UScriptStruct*> ToRemove;
+		auto SyncTagFromState = [&](bool bCondition, UScriptStruct* Tag)
+		{
+			if (bCondition)
+			{
+				ToAdd.Add(Tag);
+			}
+			else
+			{
+				ToRemove.Add(Tag);
+			}
+		};
+		bool bAnyStatus = false;
+		auto SyncStatusFromState = [&](bool bCondition, ESCCModification Modification)
+		{
+			if (bCondition)
+			{
+				DataStorage->AddOrGetColumn<FSCCStatusColumn>(Row)->Modification = Modification;
+				bAnyStatus = true;
+			}
+		};
+		SyncStatusFromState(State.IsModified(), ESCCModification::Modified);
+		SyncStatusFromState(State.IsAdded(), ESCCModification::Added);
+		SyncStatusFromState(State.IsDeleted(), ESCCModification::Removed);
+		SyncStatusFromState(State.IsConflicted(), ESCCModification::Conflicted);
+		if (!bAnyStatus)
+		{
+			ToRemove.Add(FSCCStatusColumn::StaticStruct());
+		}
+						
+		FString WhoCheckedOut;
+		const bool bIsCheckedOutByOther = State.IsCheckedOutOther(&WhoCheckedOut);
+		if (bIsCheckedOutByOther && State.bExclusiveCheckout)
+		{
+			DataStorage->AddOrGetColumn<FSCCExternallyLockedColumn>(Row)->LockedBy.Name = WhoCheckedOut;
+		}
+		else
+		{
+			ToRemove.Add(FSCCExternallyLockedColumn::StaticStruct());
+		}
+		SyncTagFromState(bIsCheckedOutByOther, FSCCExternallyEditedTag::StaticStruct());
+		SyncTagFromState(State.IsCheckedOut(), FSCCLockedTag::StaticStruct());
+		SyncTagFromState(State.Changelist.IsInitialized(), FSCCInChangelistTag::StaticStruct());
+		DataStorage->AddRemoveColumns(Row, ToAdd, ToRemove);
+	};
+
 	bool bUpdated = false;
 
 	const FDateTime Now = FDateTime::Now();
@@ -1995,6 +2080,7 @@ bool FPerforceUpdateStatusWorker::UpdateStates() const
 		State->History = MoveTemp(History);
 		State->TimeStamp = Now;
 		bUpdated = true;
+		UpdateDataStorage(*State);
 	}
 
 	// next update state from 'opened' call
@@ -2008,6 +2094,7 @@ bool FPerforceUpdateStatusWorker::UpdateStates() const
 		State->History = History;
 		State->TimeStamp = Now;
 		bUpdated = true;
+		UpdateDataStorage(*State);
 	}
 
 	// add modified state
@@ -2018,6 +2105,7 @@ bool FPerforceUpdateStatusWorker::UpdateStates() const
 		State->bModifed = true;
 		State->TimeStamp = Now;
 		bUpdated = true;
+		UpdateDataStorage(*State);
 	}
 
 	return !bForceQuiet && bUpdated;
diff --git a/Engine/Plugins/Experimental/TypedElementsDataStorage/Source/TypedElementsDataStorage/Private/Compatibility/TypedElementSCCProcessors.cpp b/Engine/Plugins/Experimental/TypedElementsDataStorage/Source/TypedElementsDataStorage/Private/Compatibility/TypedElementSCCProcessors.cpp
--- a/Engine/Plugins/Experimental/TypedElementsDataStorage/Source/TypedElementsDataStorage/Private/Compatibility/TypedElementSCCProcessors.cpp
+++ b/Engine/Plugins/Experimental/TypedElementsDataStorage/Source/TypedElementsDataStorage/Private/Compatibility/TypedElementSCCProcessors.cpp
@@ -5,6 +5,7 @@
 #include "HAL/IConsoleManager.h"
 
 #include "Elements/Columns/TypedElementMiscColumns.h"
+#include "Elements/Columns/TypedElementPackageColumns.h"
 #include "Elements/Columns/TypedElementRevisionControlColumns.h"
 #include "Elements/Columns/TypedElementSelectionColumns.h"
 #include "Elements/Columns/TypedElementViewportColumns.h"
@@ -71,6 +72,15 @@ FAutoConsoleCommandWithArgsAndOutputDevice SetSelectionSCCStateConsoleCommand(
 		}
 	));
 
+void UTypedElementSCCFactory::RegisterTables(ITypedElementDataStorageInterface& DataStorage) const
+{
+	DataStorage.RegisterTable(
+		TTypedElementColumnTypeList<
+			FTypedElementPackagePathColumn, FTypedElementPackageLoadedPathColumn,
+			FSCCRevisionIdColumn, FSCCExternalRevisionIdColumn>(),
+		FName("Editor_RevisionControlTable"));
+}
+
 void UTypedElementSCCFactory::RegisterQueries(ITypedElementDataStorageInterface& DataStorage) const
 {
 	using namespace TypedElementQueryBuilder;
diff --git a/Engine/Plugins/Experimental/TypedElementsDataStorage/Source/TypedElementsDataStorage/Private/Compatibility/TypedElementSCCProcessors.h b/Engine/Plugins/Experimental/TypedElementsDataStorage/Source/TypedElementsDataStorage/Private/Compatibility/TypedElementSCCProcessors.h
--- a/Engine/Plugins/Experimental/TypedElementsDataStorage/Source/TypedElementsDataStorage/Private/Compatibility/TypedElementSCCProcessors.h
+++ b/Engine/Plugins/Experimental/TypedElementsDataStorage/Source/TypedElementsDataStorage/Private/Compatibility/TypedElementSCCProcessors.h
@@ -15,6 +15,7 @@ class UTypedElementSCCFactory : public UTypedElementDataStorageFactory
 public:
 	~UTypedElementSCCFactory() override = default;
 
+	void RegisterTables(ITypedElementDataStorageInterface& DataStorage) const override;
 	void RegisterQueries(ITypedElementDataStorageInterface& DataStorage) const override;
 	
 };
diff --git a/Engine/Plugins/Experimental/TypedElementsDataStorage/Source/TypedElementsDataStorage/Private/Compatibility/TypedElementUObjectPackagePathToColumnQueries.cpp b/Engine/Plugins/Experimental/TypedElementsDataStorage/Source/TypedElementsDataStorage/Private/Compatibility/TypedElementUObjectPackagePathToColumnQueries.cpp
--- a/Engine/Plugins/Experimental/TypedElementsDataStorage/Source/TypedElementsDataStorage/Private/Compatibility/TypedElementUObjectPackagePathToColumnQueries.cpp
+++ b/Engine/Plugins/Experimental/TypedElementsDataStorage/Source/TypedElementsDataStorage/Private/Compatibility/TypedElementUObjectPackagePathToColumnQueries.cpp
@@ -4,28 +4,88 @@
 
 #include "Elements/Columns/TypedElementCompatibilityColumns.h"
 #include "Elements/Columns/TypedElementPackageColumns.h"
+#include "Elements/Framework/TypedElementIndexHasher.h"
 #include "Elements/Framework/TypedElementQueryBuilder.h"
 #include "UObject/Package.h"
 
+static void ResolvePackageReference(ITypedElementDataStorageInterface::IQueryContext& Context, const UPackage* Package, TypedElementRowHandle Row, TypedElementRowHandle PackageRow)
+{
+	FTypedElementPackageReference PackageReference;
+	PackageReference.Row = PackageRow;
+	Context.AddColumn(Row, MoveTemp(PackageReference));
+				
+	FTypedElementPackagePathColumn PathColumn;
+	FTypedElementPackageLoadedPathColumn LoadedPathColumn;
+		
+	Package->GetPathName(nullptr, PathColumn.Path);
+	LoadedPathColumn.LoadedPath = Package->GetLoadedPath();
+				
+	Context.AddColumn(PackageRow, MoveTemp(PathColumn));
+	Context.AddColumn(PackageRow, MoveTemp(LoadedPathColumn));	
+};
+
 void UTypedElementUObjectPackagePathFactory::RegisterQueries(ITypedElementDataStorageInterface& DataStorage) const
 {
 	using namespace TypedElementQueryBuilder;
+	using DSI = ITypedElementDataStorageInterface;
 	
 	DataStorage.RegisterQuery(
 		Select(
 			TEXT("Sync UObject package info to columns"),
-			FObserver(FObserver::EEvent::Add, FTypedElementUObjectColumn::StaticStruct())
+			FObserver::OnAdd<FTypedElementUObjectColumn>()
 				.ForceToGameThread(true),
-			[](const FTypedElementUObjectColumn& Object, FTypedElementPackagePathColumn& Path, FTypedElementPackageLoadedPathColumn& LoadedPath)
+			[](DSI::IQueryContext& Context, TypedElementRowHandle Row, const FTypedElementUObjectColumn& Object)
 			{
 				if (const UObject* ObjectInstance = Object.Object.Get(); ObjectInstance != nullptr)
 				{
 					const UPackage* Target = ObjectInstance->GetPackage();
-					Target->GetPathName(nullptr, Path.Path);
-					LoadedPath.LoadedPath = Target->GetLoadedPath();
+					FString Path;
+					Target->GetPathName(nullptr, Path);
+					if (FString PackageFilename; FPackageName::TryConvertLongPackageNameToFilename(Path, PackageFilename))
+					{
+						FPaths::NormalizeFilename(PackageFilename);
+						TypedElementDataStorage::IndexHash Index = TypedElementDataStorage::GenerateIndexHash(FPaths::ConvertRelativePathToFull(PackageFilename));
+						TypedElementRowHandle PackageRow = Context.FindIndexedRow(Index);
+						if (Context.IsRowAvailable(PackageRow))
+						{
+							const UPackage* Package = ObjectInstance->GetPackage();
+							ResolvePackageReference(Context, Package, Row, PackageRow);
+						}
+						else
+						{
+							FTypedElementPackageUnresolvedReference UnresolvedPackageReference;
+                            UnresolvedPackageReference.Index = Index;
+                            Context.AddColumn(Row, MoveTemp(UnresolvedPackageReference));
+						}
+					}
 				}
 			}
 		)
 		.Compile()
 	);
+
+	DataStorage.RegisterQuery(
+		Select(
+			TEXT("Resolve package references"),
+			FProcessor(DSI::EQueryTickPhase::DuringPhysics, DataStorage.GetQueryTickGroupName(DSI::EQueryTickGroups::SyncExternalToDataStorage)),
+			[](DSI::IQueryContext& Context, TypedElementRowHandle Row, const FTypedElementUObjectColumn& Object, const FTypedElementPackageUnresolvedReference& UnresolvedPackageReference)
+			{
+				TypedElementRowHandle PackageRow = Context.FindIndexedRow(UnresolvedPackageReference.Index);
+				if (!Context.IsRowAvailable(PackageRow))
+				{
+					return;
+				}
+				const UObject* ObjectInstance = Object.Object.Get();
+				if (!ObjectInstance)
+				{
+					return;
+				}
+				const UPackage* Package = ObjectInstance->GetPackage();
+				Context.RemoveColumns(Row, { FTypedElementPackageUnresolvedReference::StaticStruct() });
+
+				ResolvePackageReference(Context, Package, Row, PackageRow);
+			}
+		)
+		.Compile()
+	);
 }
diff --git a/Engine/Plugins/Experimental/TypedElementsDataStorage/Source/TypedElementsDataStorage/Private/TypedElementDatabaseCompatibility.cpp b/Engine/Plugins/Experimental/TypedElementsDataStorage/Source/TypedElementsDataStorage/Private/TypedElementDatabaseCompatibility.cpp
--- a/Engine/Plugins/Experimental/TypedElementsDataStorage/Source/TypedElementsDataStorage/Private/TypedElementDatabaseCompatibility.cpp
+++ b/Engine/Plugins/Experimental/TypedElementsDataStorage/Source/TypedElementsDataStorage/Private/TypedElementDatabaseCompatibility.cpp
@@ -319,7 +319,6 @@ void UTypedElementDatabaseCompatibility::CreateStandardArchetypes()
 	StandardActorTable = Storage->RegisterTable(TTypedElementColumnTypeList<
 			FMassActorFragment, FTypedElementUObjectColumn, FTypedElementClassTypeInfoColumn,
 			FTypedElementLabelColumn, FTypedElementLabelHashColumn,
-			FTypedElementPackagePathColumn, FTypedElementPackageLoadedPathColumn,
 			FTypedElementSyncFromWorldTag>(), 
 		FName("Editor_StandardActorTable"));
 
@@ -329,7 +328,6 @@ void UTypedElementDatabaseCompatibility::CreateStandardArchetypes()
 
 	StandardUObjectTable = Storage->RegisterTable(TTypedElementColumnTypeList<
 			FTypedElementUObjectColumn, FTypedElementClassTypeInfoColumn,
-			FTypedElementPackagePathColumn, FTypedElementPackageLoadedPathColumn,
 			FTypedElementSyncFromWorldTag>(), 
 		FName("Editor_StandardUObjectTable"));
 
diff --git a/Engine/Source/Runtime/TypedElementFramework/Public/Elements/Columns/TypedElementPackageColumns.h b/Engine/Source/Runtime/TypedElementFramework/Public/Elements/Columns/TypedElementPackageColumns.h
--- a/Engine/Source/Runtime/TypedElementFramework/Public/Elements/Columns/TypedElementPackageColumns.h
+++ b/Engine/Source/Runtime/TypedElementFramework/Public/Elements/Columns/TypedElementPackageColumns.h
@@ -10,6 +10,27 @@
 
 #include "TypedElementPackageColumns.generated.h"
 
+/**
+ * A package reference column that has not yet been resolved to reference a package.
+ */
+USTRUCT(meta = (DisplayName = "Unresolved package path reference"))
+struct FTypedElementPackageUnresolvedReference final : public FTypedElementDataStorageColumn
+{
+	GENERATED_BODY()
+
+	TypedElementDataStorage::IndexHash Index;
+};
+
+/**
+ * Column that references a row in the table that provides package and source control information.
+ */
+USTRUCT(meta = (DisplayName = "Package path reference"))
+struct FTypedElementPackageReference final : public FTypedElementDataStorageColumn
+{
+	GENERATED_BODY()
+	
+	TypedElementRowHandle Row;
+};
 
 /**
  * Column that stores the path of a package.
