Fix ensure firing in SCC integration

diff --git a/Engine/Plugins/Developer/PerforceSourceControl/Source/PerforceSourceControl/Private/PerforceSourceControlOperations.cpp b/Engine/Plugins/Developer/PerforceSourceControl/Source/PerforceSourceControl/Private/PerforceSourceControlOperations.cpp
--- a/Engine/Plugins/Developer/PerforceSourceControl/Source/PerforceSourceControl/Private/PerforceSourceControlOperations.cpp
+++ b/Engine/Plugins/Developer/PerforceSourceControl/Source/PerforceSourceControl/Private/PerforceSourceControlOperations.cpp
@@ -22,6 +22,7 @@
 #include "PerforceSourceControlRevision.h"
 #include "SourceControlHelpers.h"
 #include "SourceControlOperations.h"
+#include "Elements/Columns/TypedElementMiscColumns.h"
 #include "Elements/Columns/TypedElementPackageColumns.h"
 #include "Elements/Columns/TypedElementRevisionControlColumns.h"
 #include "Elements/Framework/TypedElementIndexHasher.h"
@@ -1989,17 +1990,23 @@ bool FPerforceUpdateStatusWorker::UpdateStates() const
 		using namespace TypedElementQueryBuilder;
 		using DSI = ITypedElementDataStorageInterface;
 
+		if (!State.IsSourceControlled())
+		{
+			return;
+		}
+
 		DSI* DataStorage = UTypedElementRegistry::GetInstance()->GetMutableDataStorage();
 		if (!DataStorage)
 		{
 			return;
 		}
-		
-		FString Filename = FPaths::SetExtension(State.GetFilename(), "");
-		FPaths::NormalizeFilename(Filename);
 
-		auto GetRevisionControlRow = [DataStorage](const FString& Filename) -> TypedElementRowHandle
+		auto GetRevisionControlRow = [DataStorage](const FString& InFilename) -> TypedElementRowHandle
 		{
+			FString Filename = FPaths::SetExtension(InFilename, "");
+			FPaths::NormalizeFilename(Filename);
+			Filename = FPaths::ConvertRelativePathToFull(Filename);
+			
 			uint64 Index = TypedElementDataStorage::GenerateIndexHash(Filename);
 			TypedElementRowHandle Row = DataStorage->FindIndexedRow(Index);
 
@@ -2012,12 +2019,12 @@ bool FPerforceUpdateStatusWorker::UpdateStates() const
 			return Row;
 		};
 
-		TypedElementRowHandle Row = GetRevisionControlRow(Filename);
+		TypedElementRowHandle Row = GetRevisionControlRow(State.GetFilename());
 
 		DataStorage->AddOrGetColumn<FSCCRevisionIdColumn>(Row)->RevisionId.Id[0] = State.LocalRevNumber;
 		DataStorage->AddOrGetColumn<FSCCExternalRevisionIdColumn>(Row)->RevisionId.Id[0] = State.DepotRevNumber;
 
-		TArray<UScriptStruct*> ToAdd;
+		TArray<UScriptStruct*> ToAdd { FTypedElementSyncFromWorldTag::StaticStruct() };
 		TArray<UScriptStruct*> ToRemove;
 		auto SyncTagFromState = [&](bool bCondition, UScriptStruct* Tag)
 		{
@@ -2061,7 +2068,12 @@ bool FPerforceUpdateStatusWorker::UpdateStates() const
 		SyncTagFromState(bIsCheckedOutByOther, FSCCExternallyEditedTag::StaticStruct());
 		SyncTagFromState(State.IsCheckedOut(), FSCCLockedTag::StaticStruct());
 		SyncTagFromState(State.Changelist.IsInitialized(), FSCCInChangelistTag::StaticStruct());
-		DataStorage->AddRemoveColumns(Row, ToAdd, ToRemove);
+		DataStorage->AddColumns(Row, ToAdd);
+		// Not using batched removal for now because it ensures that all are present
+		for (UScriptStruct* Column : ToRemove)
+		{
+			DataStorage->RemoveColumn(Row, Column);
+		}
 	};
 
 	bool bUpdated = false;
