Add support for EStatus::WorkspacePath with the perforce source control provider.

#rnx

- The safest way to determine the client spec root is to just initialize a p4 connection and extract the info from that.

diff --git a/Engine/Plugins/Developer/PerforceSourceControl/Source/PerforceSourceControl/Private/PerforceSourceControlProvider.cpp b/Engine/Plugins/Developer/PerforceSourceControl/Source/PerforceSourceControl/Private/PerforceSourceControlProvider.cpp
--- a/Engine/Plugins/Developer/PerforceSourceControl/Source/PerforceSourceControl/Private/PerforceSourceControlProvider.cpp
+++ b/Engine/Plugins/Developer/PerforceSourceControl/Source/PerforceSourceControl/Private/PerforceSourceControlProvider.cpp
@@ -180,6 +180,17 @@ TMap<ISourceControlProvider::EStatus, FString> FPerforceSourceControlProvider::G
 	Result.Add(EStatus::Port, Settings.GetPort());
 	Result.Add(EStatus::User, Settings.GetUserName());
 	Result.Add(EStatus::Client, Settings.GetWorkspace());
+
+	if (!Settings.GetWorkspace().IsEmpty())
+	{
+		FScopedPerforceConnection ScopedConnection(EConcurrency::Synchronous, *const_cast<FPerforceSourceControlProvider*>(this));
+		if (ScopedConnection.IsValid())
+		{
+			FPerforceConnection& Connection = ScopedConnection.GetConnection();
+			Result.Add(EStatus::WorkspacePath, Connection.ClientRoot);
+		}
+	}
+
 	return Result;
 }
 
