Fix an assert seen when errors are encountered while asynchronously updating the change list view window.

#rnx

- FSourceControlAssetDataCache can launch work on the taskgraph system in the method ::FetchAssetData which in turn calls ISourceControlRevision::Get.
- The perforce implementation of this uses FMessageLog to report errors but this can only be used safely on the GameThread. When used on the taskgraph system we run the risk of either triggering an unsafe delegate assert or we end up trying to update a slate widget (if the FMessageLog window is open) which will assert when done outside of the GameThread.
- The short term fix is to use FTSMessageLog which is a quick and dirty wrapper in the perforce module that acts like FMessageLog but when called on a taskgraph thread it will just log the message rather than trying to route it to the FMessageLog system.
-- There is a work item to make FMessageLog thread safe but due to some API choices this is non-trivial.
- Note that this change does not address the original errors, but should prevent the crash.

diff --git a/Engine/Plugins/Developer/PerforceSourceControl/Source/PerforceSourceControl/Private/PerforceSourceControlRevision.cpp b/Engine/Plugins/Developer/PerforceSourceControl/Source/PerforceSourceControl/Private/PerforceSourceControlRevision.cpp
--- a/Engine/Plugins/Developer/PerforceSourceControl/Source/PerforceSourceControl/Private/PerforceSourceControlRevision.cpp
+++ b/Engine/Plugins/Developer/PerforceSourceControl/Source/PerforceSourceControl/Private/PerforceSourceControlRevision.cpp
@@ -6,6 +6,7 @@
 #include "Misc/FileHelper.h"
 #include "Misc/Paths.h"
 #include "PerforceConnection.h"
+#include "PerforceMessageLog.h"
 #include "PerforceSourceControlProvider.h"
 
 #define LOCTEXT_NAMESPACE "PerforceSourceControl"
@@ -92,7 +93,7 @@ bool FPerforceSourceControlRevision::Get(FString& InOutFilename, FPerforceConnec
 	{
 		for (auto Iter(ErrorMessages.CreateConstIterator()); Iter; Iter++)
 		{
-			FMessageLog("SourceControl").Error(FText::Format(LOCTEXT("PerforceSourceControlRevisionGetErrorFormat", "FPerforceSourceControlRevision::Get print Error: {0}"), *Iter));
+			FTSMessageLog("SourceControl").Error(FText::Format(LOCTEXT("PerforceSourceControlRevisionGetErrorFormat", "FPerforceSourceControlRevision::Get print Error: {0}"), *Iter));
 		}
 
 		return false;
