Fix SOURCE_CONTROL_WITH_SLATE to work properly with PerforceSourceControl module.

#rnx

- Remove the slate widgets from the module if we are not compiling with slate.

diff --git a/Engine/Plugins/Developer/PerforceSourceControl/Source/PerforceSourceControl/PerforceSourceControl.Build.cs b/Engine/Plugins/Developer/PerforceSourceControl/Source/PerforceSourceControl/PerforceSourceControl.Build.cs
--- a/Engine/Plugins/Developer/PerforceSourceControl/Source/PerforceSourceControl/PerforceSourceControl.Build.cs
+++ b/Engine/Plugins/Developer/PerforceSourceControl/Source/PerforceSourceControl/PerforceSourceControl.Build.cs
@@ -12,14 +12,23 @@ public class PerforceSourceControl : ModuleRules
 			new string[] {
 				"Core",
 				"CoreUObject",
-				"InputCore",
-				"Slate",
-				"SlateCore",
 				"SourceControl",
 				"TypedElementFramework",
 			}
 		);
 
+		// See SOURCE_CONTROL_WITH_SLATE in SourceControl.Build.cs
+		if (Target.bUsesSlate)
+		{
+			PrivateDependencyModuleNames.AddRange(
+				new string[] {
+					"InputCore",
+					"Slate",
+					"SlateCore"
+				}
+			);
+		}
+
 		AddEngineThirdPartyPrivateStaticDependencies(Target, "Perforce");
 
 		if (Target.Platform == UnrealTargetPlatform.Win64 || Target.Platform == UnrealTargetPlatform.Mac)
diff --git a/Engine/Plugins/Developer/PerforceSourceControl/Source/PerforceSourceControl/Private/PerforceSourceControlProvider.cpp b/Engine/Plugins/Developer/PerforceSourceControl/Source/PerforceSourceControl/Private/PerforceSourceControlProvider.cpp
--- a/Engine/Plugins/Developer/PerforceSourceControl/Source/PerforceSourceControl/Private/PerforceSourceControlProvider.cpp
+++ b/Engine/Plugins/Developer/PerforceSourceControl/Source/PerforceSourceControl/Private/PerforceSourceControlProvider.cpp
@@ -3,6 +3,7 @@
 #include "PerforceSourceControlProvider.h"
 
 #include "Algo/Transform.h"
+#include "HAL/IConsoleManager.h"
 #include "HAL/PlatformProcess.h"
 #include "ISourceControlLabel.h"
 #include "ISourceControlModule.h"
@@ -19,7 +20,6 @@
 #include "ScopedSourceControlProgress.h"
 #include "SourceControlHelpers.h"
 #include "SourceControlOperations.h"
-#include "Widgets/DeclarativeSyntaxSupport.h"
 
 static FName ProviderName("Perforce");
 
diff --git a/Engine/Plugins/Developer/PerforceSourceControl/Source/PerforceSourceControl/Private/PerforceSourceControlSettings.cpp b/Engine/Plugins/Developer/PerforceSourceControl/Source/PerforceSourceControl/Private/PerforceSourceControlSettings.cpp
--- a/Engine/Plugins/Developer/PerforceSourceControl/Source/PerforceSourceControl/Private/PerforceSourceControlSettings.cpp
+++ b/Engine/Plugins/Developer/PerforceSourceControl/Source/PerforceSourceControl/Private/PerforceSourceControlSettings.cpp
@@ -185,8 +185,13 @@ FPerforceConnectionInfo FPerforceSourceControlSettings::GetConnectionInfo() cons
 		OutConnectionInfo = ConnectionInfo;
 	}
 	
+#if SOURCE_CONTROL_WITH_SLATE
 	// The password needs to be gotten straight from the input UI, its not stored anywhere else
 	const FString Password = SPerforceSourceControlSettings::GetPassword();
+#else
+	const FString Password;
+#endif //SOURCE_CONTROL_WITH_SLATE
+
 	if(!Password.IsEmpty())
 	{
 		OutConnectionInfo.Password = Password;
diff --git a/Engine/Plugins/Developer/PerforceSourceControl/Source/PerforceSourceControl/Private/PerforceSourceControlState.cpp b/Engine/Plugins/Developer/PerforceSourceControl/Source/PerforceSourceControl/Private/PerforceSourceControlState.cpp
--- a/Engine/Plugins/Developer/PerforceSourceControl/Source/PerforceSourceControl/Private/PerforceSourceControlState.cpp
+++ b/Engine/Plugins/Developer/PerforceSourceControl/Source/PerforceSourceControl/Private/PerforceSourceControlState.cpp
@@ -3,8 +3,11 @@
 #include "PerforceSourceControlState.h"
 #include "PerforceSourceControlRevision.h"
 #include "Misc/EngineVersion.h"
-#include "Styling/AppStyle.h"
+
+#if SOURCE_CONTROL_WITH_SLATE
+#include "Textures/SlateIcon.h"
 #include "RevisionControlStyle/RevisionControlStyle.h"
+#endif //SOURCE_CONTROL_WITH_SLATE
 
 #define LOCTEXT_NAMESPACE "PerforceSourceControl.State"
 
diff --git a/Engine/Plugins/Developer/PerforceSourceControl/Source/PerforceSourceControl/Private/SPerforceSourceControlSettings.cpp b/Engine/Plugins/Developer/PerforceSourceControl/Source/PerforceSourceControl/Private/SPerforceSourceControlSettings.cpp
--- a/Engine/Plugins/Developer/PerforceSourceControl/Source/PerforceSourceControl/Private/SPerforceSourceControlSettings.cpp
+++ b/Engine/Plugins/Developer/PerforceSourceControl/Source/PerforceSourceControl/Private/SPerforceSourceControlSettings.cpp
@@ -2,6 +2,8 @@
 
 #include "SPerforceSourceControlSettings.h"
 
+#if SOURCE_CONTROL_WITH_SLATE
+
 #include "ISourceControlModule.h"
 #include "PerforceSourceControlInternalOperations.h"
 #include "PerforceSourceControlPrivate.h"
@@ -476,3 +478,5 @@ FPerforceSourceControlProvider& SPerforceSourceControlSettings::GetSCCProvider()
 }
 
 #undef LOCTEXT_NAMESPACE
+
+#endif //SOURCE_CONTROL_WITH_SLATE
diff --git a/Engine/Plugins/Developer/PerforceSourceControl/Source/PerforceSourceControl/Private/SPerforceSourceControlSettings.h b/Engine/Plugins/Developer/PerforceSourceControl/Source/PerforceSourceControl/Private/SPerforceSourceControlSettings.h
--- a/Engine/Plugins/Developer/PerforceSourceControl/Source/PerforceSourceControl/Private/SPerforceSourceControlSettings.h
+++ b/Engine/Plugins/Developer/PerforceSourceControl/Source/PerforceSourceControl/Private/SPerforceSourceControlSettings.h
@@ -2,6 +2,8 @@
 
 #pragma once
 
+#if SOURCE_CONTROL_WITH_SLATE
+
 #include "CoreMinimal.h"
 #include "Layout/Visibility.h"
 #include "Input/Reply.h"
@@ -150,3 +152,5 @@ private:
 	/** Whether to expand advanced settings */
 	bool bAreAdvancedSettingsExpanded;
 };
+
+#endif //SOURCE_CONTROL_WITH_SLATE
