Capture info messages for all perforce commands, the calling code can then decide if the info is important or not.

### Perforce Connection
- FP4ClientUser now takes a FSourceControlResultInfo structure to store the warning/error messages generated during a perforce command, in addition we can now store info messages.
-- In almost all cases the array that was being passed in to store the warning/error messages came from a FSourceControlResultInfo so it doesn't change the overall flow.
- Rather than override ClientUser::OutputInfo to get the info messages we have overridden FP4ClientUser::Message instead as this is the preferred way in the cpp p4 api to capturing messages.
-- One difference between ::Message and the older capturing methods is that the older methods seemed to result in a newline character at the end of the message, which we do not see with the newer capture method. As some of our result parsing code relies on this behavior FP4ClientUser::Message will append this newline character for now until the parsing code can be fixed.
- For now FP4ClientUser still overrides ::OutputInfo, ::OutputError and ::HandleError but only to add a checkNoEntry() to each one. This is so that we can confirm that there are no remaining uses of it before removing it entirely.
- All classes derived from FP4ClientUser have been updated to work with the changes. In some cases this meant changing the override of ClientUser::OutputInfo to ClientUser::Message.
-- With a little bit of work we could probably remove these derived classes now that FP4ClientUser supports capturing the info messages which would help simplify the code base however this should be done in a future work item.

### Misc
- Remove ::FPerforceSourceControlProvider::GetWorkspaceList rather than update it, as it is not used.
- Add FSourceControlResultInfo::HasErrors to check if the struct contains errors or not.
-- Did not replace existing use of FPerforceSourceControlProvider::ErrorMessages.IsEmpty() that should be done as it's own work item.

diff --git a/Engine/Plugins/Developer/PerforceSourceControl/Source/PerforceSourceControl/Private/PerforceConnection.cpp b/Engine/Plugins/Developer/PerforceSourceControl/Source/PerforceSourceControl/Private/PerforceConnection.cpp
--- a/Engine/Plugins/Developer/PerforceSourceControl/Source/PerforceSourceControl/Private/PerforceConnection.cpp
+++ b/Engine/Plugins/Developer/PerforceSourceControl/Source/PerforceSourceControl/Private/PerforceConnection.cpp
@@ -157,11 +157,11 @@ class FP4ClientUser : public ClientUser
 {
 public:
 
-	FP4ClientUser(FP4RecordSet& InRecords, EP4ClientUserFlags InFlags, TArray<FText>& InOutErrorMessages)
+	FP4ClientUser(FP4RecordSet& InRecords, EP4ClientUserFlags InFlags, FSourceControlResultInfo& OutResultInfo)
 		: ClientUser()
 		, Flags(InFlags)
 		, Records(InRecords)
-		, OutErrorMessages(InOutErrorMessages)
+		, ResultInfo(OutResultInfo)
 	{
 
 	}
@@ -185,7 +185,7 @@ public:
 				FText Message = FText::Format(LOCTEXT("P4Client_GatheringUnfinished", "Started gathering depot file '{0}' before the previous file finished!"),
 					FText::FromString(File.GetDepotPath()));
 
-				OutErrorMessages.Add(MoveTemp(Message));
+				ResultInfo.ErrorMessages.Add(MoveTemp(Message));
 			}
 
 			File.Initialize(Record);
@@ -214,7 +214,7 @@ public:
 					File.GetRemainingBytes(),
 					FText::FromString(File.GetDepotPath()));
 
-				OutErrorMessages.Add(MoveTemp(Message));
+				ResultInfo.ErrorMessages.Add(MoveTemp(Message));
 
 				File.Reset();
 			}
@@ -245,7 +245,7 @@ public:
 						File.GetRemainingBytes(),
 						FText::FromString(File.GetDepotPath()));
 
-					OutErrorMessages.Add(MoveTemp(Message));
+					ResultInfo.ErrorMessages.Add(MoveTemp(Message));
 
 					File.Reset();
 				}
@@ -258,7 +258,7 @@ public:
 					File.GetRemainingBytes(),
 					FText::FromString(File.GetDepotPath()));
 
-				OutErrorMessages.Add(MoveTemp(Message));
+				ResultInfo.ErrorMessages.Add(MoveTemp(Message));
 
 				File.Reset();
 			}
@@ -269,12 +269,46 @@ public:
 		}
 	}
 
-	/** Called by P4API when an error message is available. */
+	virtual void Message(Error* err) override
+	{
+		StrBuf Buffer;
+		err->Fmt(Buffer, EF_PLAIN);
+
+		FString Message(TO_TCHAR(Buffer.Text(), IsUnicodeServer()));
+
+		// Previously we used ::HandleError which would have \n at the end of each line.
+		// For now we should add that to maintain compatibility with existing code.
+		if (!Message.EndsWith(TEXT("\n")))
+		{
+			Message.Append(TEXT("\n"));
+		}
+
+		if (err->GetSeverity() <= ErrorSeverity::E_INFO)
+		{
+			ResultInfo.InfoMessages.Add(FText::FromString(MoveTemp(Message)));
+		}
+		else
+		{
+			ResultInfo.ErrorMessages.Add(FText::FromString(MoveTemp(Message)));
+		}
+	}
+
+	virtual void OutputInfo(char Indent, const char* InInfo) override
+	{
+		// Should be caught by ::Message
+		checkNoEntry();
+	}
+
+	virtual void OutputError(const char* errBuf) override
+	{
+		// Should be caught by ::Message
+		checkNoEntry();
+	}
+
 	virtual void HandleError(Error* InError) override
 	{
-		StrBuf ErrorMessage;
-		InError->Fmt(&ErrorMessage);
-		OutErrorMessages.Add(FText::FromString(FString(TO_TCHAR(ErrorMessage.Text(), IsUnicodeServer()))));
+		// Should be caught by ::Message
+		checkNoEntry();
 	}
 
 	inline bool IsUnicodeServer() const
@@ -295,7 +329,7 @@ public:
 
 	EP4ClientUserFlags Flags;
 	FP4RecordSet& Records;
-	TArray<FText>& OutErrorMessages;
+	FSourceControlResultInfo& ResultInfo;
 
 private:
 	TArray<FSharedBuffer> Files;
@@ -307,8 +341,8 @@ private:
 class FP4CreateChangelistClientUser : public FP4ClientUser
 {
 public:
-	FP4CreateChangelistClientUser(FP4RecordSet& InRecords, EP4ClientUserFlags InFlags, TArray<FText>& InOutErrorMessages, const FText& InDescription, const TArray<FString>& InFiles, ClientApi &InP4Client)
-		:	FP4ClientUser(InRecords, InFlags, InOutErrorMessages)
+	FP4CreateChangelistClientUser(FP4RecordSet& InRecords, EP4ClientUserFlags InFlags, FSourceControlResultInfo& OutResultInfo, const FText& InDescription, const TArray<FString>& InFiles, ClientApi &InP4Client)
+		:	FP4ClientUser(InRecords, InFlags, OutResultInfo)
 		,	Description(InDescription)
 		,	ChangelistNumber(0)
 		,	Files(InFiles)
@@ -317,15 +351,27 @@ public:
 	}
 
 	/** Called by P4API when the changelist is created. */
-	virtual void OutputInfo(ANSICHAR Level, const ANSICHAR *Data) override
+	virtual void Message(Error* err) override
 	{
-		const int32 ChangeTextLen = FCString::Strlen(TEXT("Change "));
-		if (FString(TO_TCHAR(Data, IsUnicodeServer())).StartsWith(TEXT("Change ")))
+		if (err->GetSeverity() <= ErrorSeverity::E_INFO)
 		{
-			ChangelistNumber = FCString::Atoi(TO_TCHAR(Data + ChangeTextLen, IsUnicodeServer()));
+			StrBuf Buffer;
+			err->Fmt(Buffer, EF_PLAIN);
+
+			FString Message(TO_TCHAR(Buffer.Text(), IsUnicodeServer()));
+
+			const int32 ChangeTextLen = FCString::Strlen(TEXT("Change "));
+			if (Message.StartsWith(TEXT("Change ")))
+			{
+				ChangelistNumber = FCString::Atoi(*Message + ChangeTextLen);
+			}
 		}
+
+		// Pass the message on as we will still want to record it
+		FP4ClientUser::Message(err);
 	}
 
+
 	/** Called by P4API on "change -i" command. OutBuffer is filled with changelist specification text. */
 	virtual void InputData(StrBuf* OutBuffer, Error* OutError) override
 	{
@@ -372,8 +418,8 @@ public:
 class FP4EditChangelistClientUser : public FP4ClientUser
 {
 public:
-	FP4EditChangelistClientUser(FP4RecordSet& OutRecords, EP4ClientUserFlags InFlags, TArray<FText>& InOutErrorMessages, const FText& InDescription, int32 InChangelistNumber, const FP4RecordSet& InRecords, ClientApi& InP4Client)
-		: FP4ClientUser(OutRecords, InFlags, InOutErrorMessages)
+	FP4EditChangelistClientUser(FP4RecordSet& OutRecords, EP4ClientUserFlags InFlags, FSourceControlResultInfo& OutResultInfo, const FText& InDescription, int32 InChangelistNumber, const FP4RecordSet& InRecords, ClientApi& InP4Client)
+		: FP4ClientUser(OutRecords, InFlags, OutResultInfo)
 		, Description(InDescription)
 		, ChangelistNumber(InChangelistNumber)
 		, P4Client(InP4Client)
@@ -382,13 +428,24 @@ public:
 	}
 
 	/** Called by P4API when the changelist is updated. */
-	virtual void OutputInfo(ANSICHAR Level, const ANSICHAR* Data) override
+	virtual void Message(Error* err) override
 	{
-		const int32 ChangeTextLen = FCString::Strlen(TEXT("Change "));
-		if (FString(TO_TCHAR(Data, IsUnicodeServer())).StartsWith(TEXT("Change ")))
+		if (err->GetSeverity() <= ErrorSeverity::E_INFO)
 		{
-			ChangelistNumber = FCString::Atoi(TO_TCHAR(Data + ChangeTextLen, IsUnicodeServer()));
+			StrBuf Buffer;
+			err->Fmt(Buffer, EF_PLAIN);
+
+			FString Message(TO_TCHAR(Buffer.Text(), IsUnicodeServer()));
+
+			const int32 ChangeTextLen = FCString::Strlen(TEXT("Change "));
+			if (Message.StartsWith(TEXT("Change ")))
+			{
+				ChangelistNumber = FCString::Atoi(*Message + ChangeTextLen);
+			}
 		}
+
+		// Pass the message on as we will still want to record it
+		FP4ClientUser::Message(err);
 	}
 
 	/** Called by P4API on "change -i" command. OutBuffer is filled with changelist specification text. */
@@ -463,8 +520,8 @@ public:
 class FP4CommandWithStdInputClientUser : public FP4ClientUser
 {
 public:
-	FP4CommandWithStdInputClientUser(FStringView InStdInput, FP4RecordSet& InRecords, EP4ClientUserFlags InFlags, TArray<FText>& InOutErrorMessages, ClientApi& InP4Client)
-		: FP4ClientUser(InRecords, InFlags, InOutErrorMessages)
+	FP4CommandWithStdInputClientUser(FStringView InStdInput, FP4RecordSet& InRecords, EP4ClientUserFlags InFlags, FSourceControlResultInfo& OutResultInfo, ClientApi& InP4Client)
+		: FP4ClientUser(InRecords, InFlags, OutResultInfo)
 		, P4Client(InP4Client)
 		, StdInput(InStdInput)
 	{
@@ -485,8 +542,8 @@ protected:
 class FP4LoginClientUser : public FP4ClientUser
 {
 public:
-	FP4LoginClientUser(const FString& InPassword, FP4RecordSet& InRecords, EP4ClientUserFlags InFlags, TArray<FText>& InOutErrorMessages)
-		:	FP4ClientUser(InRecords, InFlags, InOutErrorMessages)
+	FP4LoginClientUser(const FString& InPassword, FP4RecordSet& InRecords, EP4ClientUserFlags InFlags, FSourceControlResultInfo& OutResultInfo)
+		:	FP4ClientUser(InRecords, InFlags, OutResultInfo)
 		,	Password(InPassword)
 	{
 	}
@@ -528,15 +585,20 @@ static bool TestLoginConnection(ClientApi& P4Client, bool bIsUnicodeServer, TArr
 {
 	TRACE_CPUPROFILER_EVENT_SCOPE(FPerforce::TestLoginConnection);
 
+	OutErrorMessages.Reset();
+
+	FSourceControlResultInfo ResultInfo;
 	FP4RecordSet Records;
 	EP4ClientUserFlags Flags = bIsUnicodeServer ? EP4ClientUserFlags::UnicodeServer : EP4ClientUserFlags::None;
-	FP4ClientUser User(Records, Flags, OutErrorMessages);
+	FP4ClientUser User(Records, Flags, ResultInfo);
 
 	const char* ArgV[] = { "-s" };
 
 	P4Client.SetArgv(1, const_cast<char* const*>(ArgV));
 	P4Client.Run("login", &User);
 
+	OutErrorMessages = MoveTemp(ResultInfo.ErrorMessages);
+
 	return OutErrorMessages.IsEmpty();
 }
 
@@ -547,10 +609,13 @@ static bool TestLoginConnection(ClientApi& P4Client, bool bIsUnicodeServer, TArr
 static bool TestClientConnection(ClientApi& P4Client, const FString& ClientSpecName, bool bIsUnicodeServer, TArray<FText>& OutErrorMessages)
 {
 	TRACE_CPUPROFILER_EVENT_SCOPE(FPerforce::TestClientConnection);
+	
+	OutErrorMessages.Reset();
 
+	FSourceControlResultInfo ResultInfo;
 	FP4RecordSet Records;
 	EP4ClientUserFlags Flags = bIsUnicodeServer ? EP4ClientUserFlags::UnicodeServer : EP4ClientUserFlags::None;
-	FP4ClientUser User(Records, Flags, OutErrorMessages);
+	FP4ClientUser User(Records, Flags, ResultInfo);
 
 	UTF8CHAR* ClientSpecUTF8Name = nullptr;
 	if(bIsUnicodeServer)
@@ -573,6 +638,8 @@ static bool TestClientConnection(ClientApi& P4Client, const FString& ClientSpecN
 	// clean up args
 	delete [] ClientSpecUTF8Name;
 
+	OutErrorMessages = MoveTemp(ResultInfo.ErrorMessages);
+
 	// If there are error messages, user name is most likely invalid. Otherwise, make sure workspace actually
 	// exists on server by checking if we have it's update date.
 	bool bConnectionOK = OutErrorMessages.Num() == 0 && Records.Num() > 0 && Records[0].Contains(TEXT("Update"));
@@ -586,11 +653,18 @@ static bool TestClientConnection(ClientApi& P4Client, const FString& ClientSpecN
 
 static bool CheckUnicodeStatus(ClientApi& P4Client, bool& bIsUnicodeServer, TArray<FText>& OutErrorMessages)
 {
+	TRACE_CPUPROFILER_EVENT_SCOPE(FPerforce::CheckUnicodeStatus);
+
+	OutErrorMessages.Reset();
+
+	FSourceControlResultInfo ResultInfo;
 	FP4RecordSet Records;
-	FP4ClientUser User(Records, EP4ClientUserFlags::None, OutErrorMessages);
+	FP4ClientUser User(Records, EP4ClientUserFlags::None, ResultInfo);
 
 	P4Client.Run("info", &User);
 
+	OutErrorMessages = MoveTemp(ResultInfo.ErrorMessages);
+
 	if(Records.Num() > 0)
 	{
 		bIsUnicodeServer = Records[0].Find(TEXT("unicode")) != nullptr;
@@ -664,9 +738,15 @@ bool FPerforceConnection::AutoDetectWorkspace(const FPerforceConnectionInfo& InC
 {
 	//before even trying to summon the window, try to "smart" connect with the default server/username
 
+	OutErrorMessages.Reset();
+
 	FPerforceConnection Connection(InConnectionInfo, SCCProvider);
 	TArray<FString> ClientSpecList;
-	Connection.GetWorkspaceList(InConnectionInfo, FOnIsCancelled(), ClientSpecList, OutErrorMessages);
+	FSourceControlResultInfo ResultInfo;
+
+	Connection.GetWorkspaceList(InConnectionInfo, FOnIsCancelled(), ClientSpecList, ResultInfo);
+
+	OutErrorMessages = MoveTemp(ResultInfo.ErrorMessages);
 
 	if (!OutErrorMessages.IsEmpty())
 	{
@@ -716,25 +796,25 @@ bool FPerforceConnection::AutoDetectWorkspace(const FPerforceConnectionInfo& InC
 
 bool FPerforceConnection::Login(const FPerforceConnectionInfo& InConnectionInfo)
 {
-	TArray<FText> ErrorMessages;
+	FSourceControlResultInfo ResultInfo;
 
 	FP4RecordSet Records;
-	FP4LoginClientUser User(InConnectionInfo.Password, Records, EP4ClientUserFlags::None, ErrorMessages);
+	FP4LoginClientUser User(InConnectionInfo.Password, Records, EP4ClientUserFlags::None, ResultInfo);
 
 	const char *ArgV[] = { "-a" };
 	P4Client.SetArgv(1, const_cast<char*const*>(ArgV));
 	P4Client.Run("login", &User);
 
-	if(ErrorMessages.Num())
+	if (ResultInfo.HasErrors())
 	{
 		UE_LOG(LogSourceControl, Error, TEXT("Login failed"));
-		for(auto ErrorMessage : ErrorMessages)
+		for (const FText& ErrorMessage : ResultInfo.ErrorMessages)
 		{
 			UE_LOG(LogSourceControl, Error, TEXT("%s"), *ErrorMessage.ToString());
 		}
 	}
 
-	return ErrorMessages.Num() == 0;
+	return !ResultInfo.HasErrors();
 }
 
 bool EnsureValidConnectionInternal(const FPerforceConnectionInfo& InSettings, FPerforceSourceControlProvider& SCCProvider, EConnectionOptions Options,
@@ -894,7 +974,7 @@ bool FPerforceConnection::EnsureValidConnection(const FPerforceConnectionInfo& I
 	return bResult;
 }
 
-bool FPerforceConnection::GetWorkspaceList(const FPerforceConnectionInfo& InConnectionInfo, FOnIsCancelled InOnIsCanceled, TArray<FString>& OutWorkspaceList, TArray<FText>& OutErrorMessages)
+bool FPerforceConnection::GetWorkspaceList(const FPerforceConnectionInfo& InConnectionInfo, FOnIsCancelled InOnIsCanceled, TArray<FString>& OutWorkspaceList, FSourceControlResultInfo& OutResultInfo)
 {
 	TRACE_CPUPROFILER_EVENT_SCOPE(FPerforceConnection::GetWorkspaceList)
 
@@ -907,7 +987,7 @@ bool FPerforceConnection::GetWorkspaceList(const FPerforceConnectionInfo& InConn
 
 		FP4RecordSet Records;
 		bool bConnectionDropped = false;
-		bool bCommandOK = RunCommand(TEXT("clients"), Params, Records, OutErrorMessages, InOnIsCanceled, bConnectionDropped);
+		bool bCommandOK = RunCommand(TEXT("clients"), Params, Records, OutResultInfo, InOnIsCanceled, bConnectionDropped);
 
 		if (bCommandOK)
 		{
@@ -1007,7 +1087,7 @@ double FPerforceConnection::GetLatestCommuncationTime() const
 }
 
 bool FPerforceConnection::RunCommand(	const FString& InCommand, const TArray<FString>& InParameters, FP4RecordSet& OutRecordSet, 
-                                        TArray<FSharedBuffer>* OutData, TArray<FText>& OutErrorMessage, 
+                                        TArray<FSharedBuffer>* OutData, FSourceControlResultInfo& OutResultInfo,
                                         FOnIsCancelled InIsCancelled, bool& OutConnectionDropped, ERunCommandFlags RunFlags)
 {
 	TRACE_CPUPROFILER_EVENT_SCOPE_TEXT(*FString::Printf(TEXT("FPerforceConnection::RunCommand_%s"), *InCommand));
@@ -1063,7 +1143,7 @@ bool FPerforceConnection::RunCommand(	const FString& InCommand, const TArray<FSt
 	ClientUserFlags |= bIsUnicode ? EP4ClientUserFlags::UnicodeServer : EP4ClientUserFlags::None;
 	ClientUserFlags |= OutData != nullptr ? EP4ClientUserFlags::CollectData : EP4ClientUserFlags::None;
 	
-	FP4ClientUser User(OutRecordSet, ClientUserFlags, OutErrorMessage);
+	FP4ClientUser User(OutRecordSet, ClientUserFlags, OutResultInfo);
 	if (EnumHasAllFlags(RunFlags, ERunCommandFlags::Quiet))
 	{
 		User.SetQuiet();
@@ -1092,7 +1172,7 @@ bool FPerforceConnection::RunCommand(	const FString& InCommand, const TArray<FSt
 	// Only report connection related errors to avoid clearing of connection related error messages
 	if (InCommand != TEXT("info"))
 	{
-		SCCProvider.SetLastErrors(OutErrorMessage);
+		SCCProvider.SetLastErrors(OutResultInfo.ErrorMessages);
 	}
 
 	if (bLogCommandDetails)
@@ -1116,7 +1196,7 @@ bool FPerforceConnection::RunCommand(	const FString& InCommand, const TArray<FSt
 	return OutRecordSet.Num() > 0;
 }
 
-int32 FPerforceConnection::CreatePendingChangelist(const FText &Description, const TArray<FString>& Files, FOnIsCancelled InIsCancelled, TArray<FText>& OutErrorMessages)
+int32 FPerforceConnection::CreatePendingChangelist(const FText &Description, const TArray<FString>& Files, FOnIsCancelled InIsCancelled, FSourceControlResultInfo& OutResultInfo)
 {
 	TArray<FString> Params;
 	FP4RecordSet Records;
@@ -1128,7 +1208,7 @@ int32 FPerforceConnection::CreatePendingChangelist(const FText &Description, con
 	FP4KeepAlive KeepAlive(InIsCancelled);
 	P4Client.SetBreak(&KeepAlive);
 
-	FP4CreateChangelistClientUser User(Records, Flags, OutErrorMessages, Description, Files, P4Client);
+	FP4CreateChangelistClientUser User(Records, Flags, OutResultInfo, Description, Files, P4Client);
 	P4Client.Run("change", &User);
 
 	P4Client.SetBreak(nullptr);
@@ -1136,7 +1216,7 @@ int32 FPerforceConnection::CreatePendingChangelist(const FText &Description, con
 	return User.ChangelistNumber;
 }
 
-int32 FPerforceConnection::EditPendingChangelist(const FText& NewDescription, int32 ChangelistNumber, FOnIsCancelled InIsCancelled, TArray<FText>& OutErrorMessages)
+int32 FPerforceConnection::EditPendingChangelist(const FText& NewDescription, int32 ChangelistNumber, FOnIsCancelled InIsCancelled, FSourceControlResultInfo& OutResultInfo)
 {
 	FP4RecordSet PreviousRecords;
 
@@ -1150,7 +1230,7 @@ int32 FPerforceConnection::EditPendingChangelist(const FText& NewDescription, in
 		// TODO : make this work also for default changelist, but should really be a Create
 		Params.Add(FString::Printf(TEXT("%d"), ChangelistNumber));
 
-		if (!RunCommand(TEXT("change"), Params, PreviousRecords, nullptr, OutErrorMessages, InIsCancelled, bConnectionDropped, RunFlags))
+		if (!RunCommand(TEXT("change"), Params, PreviousRecords, nullptr, OutResultInfo, InIsCancelled, bConnectionDropped, RunFlags))
 		{
 			return 0;
 		}
@@ -1168,16 +1248,16 @@ int32 FPerforceConnection::EditPendingChangelist(const FText& NewDescription, in
 		FP4KeepAlive KeepAlive(InIsCancelled);
 		P4Client.SetBreak(&KeepAlive);
 
-		FP4EditChangelistClientUser User(Records, Flags, OutErrorMessages, NewDescription, ChangelistNumber, PreviousRecords, P4Client);
+		FP4EditChangelistClientUser User(Records, Flags, OutResultInfo, NewDescription, ChangelistNumber, PreviousRecords, P4Client);
 		P4Client.Run("change", &User);
 
 		P4Client.SetBreak(nullptr);
 
-		return (OutErrorMessages.Num() == 0 ? User.ChangelistNumber : 0);
+		return (!OutResultInfo.HasErrors() ? User.ChangelistNumber : 0);
 	}
 }
 
-bool FPerforceConnection::CreateWorkspace(FStringView WorkspaceSpec, FOnIsCancelled InIsCancelled, TArray<FText>& OutErrorMessages)
+bool FPerforceConnection::CreateWorkspace(FStringView WorkspaceSpec, FOnIsCancelled InIsCancelled, FSourceControlResultInfo& OutResultInfo)
 {
 	TRACE_CPUPROFILER_EVENT_SCOPE(FPerforceConnection::CreateWorkspace);
 
@@ -1192,7 +1272,7 @@ bool FPerforceConnection::CreateWorkspace(FStringView WorkspaceSpec, FOnIsCancel
 	FP4KeepAlive KeepAlive(InIsCancelled);
 	P4Client.SetBreak(&KeepAlive);
 
-	FP4CommandWithStdInputClientUser User(WorkspaceSpec, Records, Flags, OutErrorMessages, P4Client);
+	FP4CommandWithStdInputClientUser User(WorkspaceSpec, Records, Flags, OutResultInfo, P4Client);
 	User.SetQuiet();	// p4 client does not return tagged output, so any output messages will be
 						// printed to stdout. Setting this will prevent that.
 
@@ -1200,7 +1280,7 @@ bool FPerforceConnection::CreateWorkspace(FStringView WorkspaceSpec, FOnIsCancel
 
 	P4Client.SetBreak(nullptr);
 	
-	return OutErrorMessages.Num() == 0;
+	return !OutResultInfo.HasErrors();
 }
 
 void FPerforceConnection::EstablishConnection(const FPerforceConnectionInfo& InConnectionInfo)
@@ -1251,14 +1331,14 @@ void FPerforceConnection::EstablishConnection(const FPerforceConnectionInfo& InC
 	else
 	{
 		TArray<FString> Params;
-		TArray<FText> ErrorMessages;
+		FSourceControlResultInfo ResultInfo;
 		FP4RecordSet Records;
 		bool bConnectionDropped = false;
 		const ERunCommandFlags RunFlags = ERunCommandFlags::DisableCommandLogging;
 
 		UE_LOG(LogSourceControl, Verbose, TEXT(" ... checking unicode status" ));
 
-		if (RunCommand(TEXT("info"), Params, Records, nullptr, ErrorMessages, FOnIsCancelled(), bConnectionDropped, RunFlags))
+		if (RunCommand(TEXT("info"), Params, Records, nullptr, ResultInfo, FOnIsCancelled(), bConnectionDropped, RunFlags))
 		{
 			// Get character encoding
 			bIsUnicode = Records[0].Find(TEXT("unicode")) != nullptr;
@@ -1291,7 +1371,7 @@ void FPerforceConnection::EstablishConnection(const FPerforceConnectionInfo& InC
 			// Gather the client root
 			UE_LOG(LogSourceControl, Verbose, TEXT(" ... getting info" ));
 			bConnectionDropped = false;
-			if (RunCommand(TEXT("info"), Params, Records, nullptr, ErrorMessages, FOnIsCancelled(), bConnectionDropped, RunFlags))
+			if (RunCommand(TEXT("info"), Params, Records, nullptr, ResultInfo, FOnIsCancelled(), bConnectionDropped, RunFlags))
 			{
 				UE_LOG(LogSourceControl, Verbose, TEXT(" ... getting clientroot" ));
 				ClientRoot = Records[0](TEXT("clientRoot"));
diff --git a/Engine/Plugins/Developer/PerforceSourceControl/Source/PerforceSourceControl/Private/PerforceConnection.h b/Engine/Plugins/Developer/PerforceSourceControl/Source/PerforceSourceControl/Private/PerforceConnection.h
--- a/Engine/Plugins/Developer/PerforceSourceControl/Source/PerforceSourceControl/Private/PerforceConnection.h
+++ b/Engine/Plugins/Developer/PerforceSourceControl/Source/PerforceSourceControl/Private/PerforceConnection.h
@@ -100,7 +100,7 @@ public:
 	 * @param OutErrorMessages	Any error messages output.
 	 * @return - True if successful
 	 */
-	bool GetWorkspaceList(const FPerforceConnectionInfo& InConnectionInfo, FOnIsCancelled InOnIsCancelled, TArray<FString>& OutWorkspaceList, TArray<FText>& OutErrorMessages);
+	bool GetWorkspaceList(const FPerforceConnectionInfo& InConnectionInfo, FOnIsCancelled InOnIsCancelled, TArray<FString>& OutWorkspaceList, FSourceControlResultInfo& OutResultInfo);
 
 	/** Returns true if connection is currently active */
 	bool IsValidConnection();
@@ -114,36 +114,36 @@ public:
 	/**
 	 * Runs internal perforce command, catches exceptions, returns results
 	 */
-	bool RunCommand(const FString& InCommand, const TArray<FString>& InParameters, FP4RecordSet& OutRecordSet, TArray<FText>&  OutErrorMessage, FOnIsCancelled InIsCancelled, bool& OutConnectionDropped)
+	bool RunCommand(const FString& InCommand, const TArray<FString>& InParameters, FP4RecordSet& OutRecordSet, FSourceControlResultInfo& OutResultInfo, FOnIsCancelled InIsCancelled, bool& OutConnectionDropped)
 	{
-		return RunCommand(InCommand, InParameters, OutRecordSet, nullptr, OutErrorMessage, InIsCancelled, OutConnectionDropped);
+		return RunCommand(InCommand, InParameters, OutRecordSet, nullptr, OutResultInfo, InIsCancelled, OutConnectionDropped);
 	}
 
 	/**
 	 * Runs internal perforce command, catches exceptions, returns results
 	 */
-	bool RunCommand(const FString& InCommand, const TArray<FString>& InParameters, FP4RecordSet& OutRecordSet, TArray<FSharedBuffer>* OutData, TArray<FText>& OutErrorMessage, FOnIsCancelled InIsCancelled, bool& OutConnectionDropped, ERunCommandFlags RunFlags = ERunCommandFlags::Default);
+	bool RunCommand(const FString& InCommand, const TArray<FString>& InParameters, FP4RecordSet& OutRecordSet, TArray<FSharedBuffer>* OutData, FSourceControlResultInfo& OutResultInfo, FOnIsCancelled InIsCancelled, bool& OutConnectionDropped, ERunCommandFlags RunFlags = ERunCommandFlags::Default);
 
 	/**
 	 * Creates a changelist with the specified description
 	 */
-	int32 CreatePendingChangelist(const FText &Description, const TArray<FString>& InFiles, FOnIsCancelled InIsCancelled, TArray<FText>& OutErrorMessages);
+	int32 CreatePendingChangelist(const FText &Description, const TArray<FString>& InFiles, FOnIsCancelled InIsCancelled, FSourceControlResultInfo& OutResultInfo);
 
 	/**
 	 * Edits a changelist with a new description
 	 */
-	int32 EditPendingChangelist(const FText& NewDescription, int32 ChangelistNumber, FOnIsCancelled InIsCancelled, TArray<FText>& OutErrorMessages);
+	int32 EditPendingChangelist(const FText& NewDescription, int32 ChangelistNumber, FOnIsCancelled InIsCancelled, FSourceControlResultInfo& OutResultInfo);
 
 	/** 
 	 * Creates a workspace based on the spec provided via the WorkspaceSpec.
 	 * 
 	 * @param WorkspaceSpec		The specification of the workspace to create
-	 * @param InIsCancelled		Delegate allowing the cancelling of the command if needed
-	 * @param OutErrorMessages	An array that will be filled with all errors encountered during the command
+	 * @param InIsCancelled		Delegate allowing the canceling of the command if needed
+	 * @param OutResultInfo		Struct that will end up containing info and error messages for the operation
 	 * 
 	 * @return Returns true if no errors were encountered, otherwise false
 	 */
-	bool CreateWorkspace(FStringView WorkspaceSpec, FOnIsCancelled InIsCancelled, TArray<FText>& OutErrorMessages);
+	bool CreateWorkspace(FStringView WorkspaceSpec, FOnIsCancelled InIsCancelled, FSourceControlResultInfo& OutResultInfo);
 
 	/**
 	 * Attempt to login - some servers will require this 
diff --git a/Engine/Plugins/Developer/PerforceSourceControl/Source/PerforceSourceControl/Private/PerforceSourceControlLabel.cpp b/Engine/Plugins/Developer/PerforceSourceControl/Source/PerforceSourceControl/Private/PerforceSourceControlLabel.cpp
--- a/Engine/Plugins/Developer/PerforceSourceControl/Source/PerforceSourceControl/Private/PerforceSourceControlLabel.cpp
+++ b/Engine/Plugins/Developer/PerforceSourceControl/Source/PerforceSourceControl/Private/PerforceSourceControlLabel.cpp
@@ -57,23 +57,22 @@ bool FPerforceSourceControlLabel::GetFileRevisions( const TArray<FString>& InFil
 		FPerforceConnection& Connection = ScopedConnection.GetConnection();
 		FP4RecordSet Records;
 		TArray<FString> Parameters;
-		TArray<FText> ErrorMessages;
+		FSourceControlResultInfo ResultInfo;
 		for(auto Iter(InFiles.CreateConstIterator()); Iter; Iter++)
 		{
 			Parameters.Add(*Iter + TEXT("@") + Name);
 		}
 		bool bConnectionDropped = false;
-		bCommandOK = Connection.RunCommand(TEXT("files"), Parameters, Records, ErrorMessages, FOnIsCancelled(), bConnectionDropped);
+		bCommandOK = Connection.RunCommand(TEXT("files"), Parameters, Records, ResultInfo, FOnIsCancelled(), bConnectionDropped);
 		if(bCommandOK)
 		{
 			ParseFilesResults(GetSCCProvider(), Records, Connection.ClientRoot, OutRevisions);
 		}
 		else
 		{
-			// output errors if any
-			for (int32 ErrorIndex = 0; ErrorIndex < ErrorMessages.Num(); ++ErrorIndex)
+			for (const FText& ErrorMsg : ResultInfo.ErrorMessages)
 			{
-				FMessageLog("SourceControl").Error(FText::Format(LOCTEXT("GetFileRevisionsErrorFormat","GetFileRevisions Error: {0}"), ErrorMessages[ErrorIndex]));
+				FMessageLog("SourceControl").Error(FText::Format(LOCTEXT("GetFileRevisionsErrorFormat","GetFileRevisions Error: {0}"), ErrorMsg));
 			}
 		}
 	}
@@ -90,20 +89,19 @@ bool FPerforceSourceControlLabel::Sync( const TArray<FString>& InFilenames ) con
 	{
 		FPerforceConnection& Connection = ScopedConnection.GetConnection();
 		FP4RecordSet Records;
-		TArray<FText> ErrorMessages;
+		FSourceControlResultInfo ResultInfo;
 		TArray<FString> Parameters;
 		for(const auto& Filename : InFilenames)
 		{
 			Parameters.Add(Filename + TEXT("@") + Name);
 		}
 		bool bConnectionDropped = false;
-		bCommandOK = Connection.RunCommand(TEXT("sync"), Parameters, Records, ErrorMessages, FOnIsCancelled(), bConnectionDropped);
+		bCommandOK = Connection.RunCommand(TEXT("sync"), Parameters, Records, ResultInfo, FOnIsCancelled(), bConnectionDropped);
 		if(!bCommandOK)
 		{
-			// output errors if any
-			for (int32 ErrorIndex = 0; ErrorIndex < ErrorMessages.Num(); ++ErrorIndex)
+			for (const FText& ErrorMsg : ResultInfo.ErrorMessages)
 			{
-				FMessageLog("SourceControl").Error(FText::Format(LOCTEXT("GetFileRevisionsSyncErrorFormat", "Sync Error: {0}"), ErrorMessages[ErrorIndex]));
+				FMessageLog("SourceControl").Error(FText::Format(LOCTEXT("GetFileRevisionsSyncErrorFormat", "Sync Error: {0}"), ErrorMsg));
 			}
 		}
 	}
diff --git a/Engine/Plugins/Developer/PerforceSourceControl/Source/PerforceSourceControl/Private/PerforceSourceControlOperations.cpp b/Engine/Plugins/Developer/PerforceSourceControl/Source/PerforceSourceControl/Private/PerforceSourceControlOperations.cpp
--- a/Engine/Plugins/Developer/PerforceSourceControl/Source/PerforceSourceControl/Private/PerforceSourceControlOperations.cpp
+++ b/Engine/Plugins/Developer/PerforceSourceControl/Source/PerforceSourceControl/Private/PerforceSourceControlOperations.cpp
@@ -417,7 +417,7 @@ bool FPerforceConnectWorker::Execute(FPerforceSourceControlCommand& InCommand)
 		FP4RecordSet Records;
 		Parameters.Add(TEXT("-o"));
 		Parameters.Add(InCommand.ConnectionInfo.Workspace);
-		InCommand.bCommandSuccessful = Connection.RunCommand(TEXT("client"), Parameters, Records, InCommand.ResultInfo.ErrorMessages, FOnIsCancelled::CreateRaw(&InCommand, &FPerforceSourceControlCommand::IsCanceled), InCommand.bConnectionDropped);
+		InCommand.bCommandSuccessful = Connection.RunCommand(TEXT("client"), Parameters, Records, InCommand.ResultInfo, FOnIsCancelled::CreateRaw(&InCommand, &FPerforceSourceControlCommand::IsCanceled), InCommand.bConnectionDropped);
 
 		// If there are error messages, user name is most likely invalid. Otherwise, make sure workspace actually
 		// exists on server by checking if we have it's update date.
@@ -482,7 +482,7 @@ bool FPerforceCheckOutWorker::Execute(FPerforceSourceControlCommand& InCommand)
 
 		Parameters.Append(InCommand.Files);
 		FP4RecordSet Records;
-		InCommand.bCommandSuccessful = Connection.RunCommand(TEXT("edit"), Parameters, Records, InCommand.ResultInfo.ErrorMessages, FOnIsCancelled::CreateRaw(&InCommand, &FPerforceSourceControlCommand::IsCanceled), InCommand.bConnectionDropped);
+		InCommand.bCommandSuccessful = Connection.RunCommand(TEXT("edit"), Parameters, Records, InCommand.ResultInfo, FOnIsCancelled::CreateRaw(&InCommand, &FPerforceSourceControlCommand::IsCanceled), InCommand.bConnectionDropped);
 		ParseRecordSetForState(Records, OutResults);
 	}
 
@@ -554,7 +554,7 @@ static bool RunReopenCommand(FPerforceSourceControlCommand& InCommand, const TAr
 				ReopenParams.Add(InFiles[FileIndex]);
 			}
 
-			bCommandSuccessful = Connection.RunCommand(TEXT("reopen"), ReopenParams, Records, InCommand.ResultInfo.ErrorMessages, FOnIsCancelled::CreateRaw(&InCommand, &FPerforceSourceControlCommand::IsCanceled), InCommand.bConnectionDropped);
+			bCommandSuccessful = Connection.RunCommand(TEXT("reopen"), ReopenParams, Records, InCommand.ResultInfo, FOnIsCancelled::CreateRaw(&InCommand, &FPerforceSourceControlCommand::IsCanceled), InCommand.bConnectionDropped);
 			if (bCommandSuccessful && OutReopenedFiles != nullptr)
 			{
 				for (int32 FileIndex = StartingIndex; FileIndex < NextIndex; FileIndex++)
@@ -614,7 +614,7 @@ bool FPerforceCheckInWorker::Execute(FPerforceSourceControlCommand& InCommand)
 
 		if (InCommand.Changelist.IsDefault())
 		{
-			int32 NewChangeList = Connection.CreatePendingChangelist(Operation->GetDescription(), TArray<FString>(), FOnIsCancelled::CreateRaw(&InCommand, &FPerforceSourceControlCommand::IsCanceled), InCommand.ResultInfo.ErrorMessages);
+			int32 NewChangeList = Connection.CreatePendingChangelist(Operation->GetDescription(), TArray<FString>(), FOnIsCancelled::CreateRaw(&InCommand, &FPerforceSourceControlCommand::IsCanceled), InCommand.ResultInfo);
 			if (NewChangeList > 0)
 			{
 				ChangeList = FPerforceSourceControlChangelist(NewChangeList);
@@ -630,7 +630,7 @@ bool FPerforceCheckInWorker::Execute(FPerforceSourceControlCommand& InCommand)
 			// Retrieves cached description to restore it in case of submit failure.
 			CachedDescription = ChangelistState->GetDescriptionText();
 
-			int32 UpdatedChangelistNumber = Connection.EditPendingChangelist(Operation->GetDescription(), InCommand.Changelist.ToInt(), FOnIsCancelled::CreateRaw(&InCommand, &FPerforceSourceControlCommand::IsCanceled), InCommand.ResultInfo.ErrorMessages);
+			int32 UpdatedChangelistNumber = Connection.EditPendingChangelist(Operation->GetDescription(), InCommand.Changelist.ToInt(), FOnIsCancelled::CreateRaw(&InCommand, &FPerforceSourceControlCommand::IsCanceled), InCommand.ResultInfo);
 			InCommand.bCommandSuccessful = (UpdatedChangelistNumber != 0) && (UpdatedChangelistNumber == InCommand.Changelist.ToInt());
 		}
 
@@ -648,7 +648,7 @@ bool FPerforceCheckInWorker::Execute(FPerforceSourceControlCommand& InCommand)
 			SubmitParams.Add(TEXT("-c"));
 			SubmitParams.Add(ChangeList.ToString());
 
-			InCommand.bCommandSuccessful = Connection.RunCommand(TEXT("submit"), SubmitParams, Records, InCommand.ResultInfo.ErrorMessages, FOnIsCancelled::CreateRaw(&InCommand, &FPerforceSourceControlCommand::IsCanceled), InCommand.bConnectionDropped);
+			InCommand.bCommandSuccessful = Connection.RunCommand(TEXT("submit"), SubmitParams, Records, InCommand.ResultInfo, FOnIsCancelled::CreateRaw(&InCommand, &FPerforceSourceControlCommand::IsCanceled), InCommand.bConnectionDropped);
 
 			if (InCommand.ResultInfo.ErrorMessages.Num() > 0)
 			{
@@ -702,13 +702,13 @@ bool FPerforceCheckInWorker::Execute(FPerforceSourceControlCommand& InCommand)
 					TArray<FString> ChangeParams;
 					ChangeParams.Add(TEXT("-d"));
 					ChangeParams.Add(ChangeList.ToString());
-					Connection.RunCommand(TEXT("change"), ChangeParams, Records, InCommand.ResultInfo.ErrorMessages, FOnIsCancelled::CreateRaw(&InCommand, &FPerforceSourceControlCommand::IsCanceled), InCommand.bConnectionDropped);
+					Connection.RunCommand(TEXT("change"), ChangeParams, Records, InCommand.ResultInfo, FOnIsCancelled::CreateRaw(&InCommand, &FPerforceSourceControlCommand::IsCanceled), InCommand.bConnectionDropped);
 				}
 			}
 			// If we modified description and we have the previous one, restore it.
 			else if ((!Operation->GetDescription().IsEmpty()) && (!CachedDescription.IsEmpty()))
 			{
-				Connection.EditPendingChangelist(CachedDescription, InCommand.Changelist.ToInt(), FOnIsCancelled::CreateRaw(&InCommand, &FPerforceSourceControlCommand::IsCanceled), InCommand.ResultInfo.ErrorMessages);
+				Connection.EditPendingChangelist(CachedDescription, InCommand.Changelist.ToInt(), FOnIsCancelled::CreateRaw(&InCommand, &FPerforceSourceControlCommand::IsCanceled), InCommand.ResultInfo);
 			}
 		}
 	}
@@ -777,7 +777,7 @@ bool FPerforceGetFileListWorker::Execute(FPerforceSourceControlCommand& InComman
 		Parameters.Append(InCommand.Files);
 		AppendMaskParameter(Parameters);
 		FP4RecordSet Records;
-		InCommand.bCommandSuccessful = Connection.RunCommand(TEXT("files"), Parameters, Records, InCommand.ResultInfo.ErrorMessages, FOnIsCancelled::CreateRaw(&InCommand, &FPerforceSourceControlCommand::IsCanceled), InCommand.bConnectionDropped);
+		InCommand.bCommandSuccessful = Connection.RunCommand(TEXT("files"), Parameters, Records, InCommand.ResultInfo, FOnIsCancelled::CreateRaw(&InCommand, &FPerforceSourceControlCommand::IsCanceled), InCommand.bConnectionDropped);
 		ParseRecordSetForState(Records, OutResults);
 
 		TArray<FString> FilesList;
@@ -853,7 +853,7 @@ bool FPerforceMarkForAddWorker::Execute(FPerforceSourceControlCommand& InCommand
 		AppendChangelistParameter(GetSCCProvider(), Parameters);
 		Parameters.Append(InCommand.Files);
 
-		InCommand.bCommandSuccessful = Connection.RunCommand(TEXT("add"), Parameters, Records, InCommand.ResultInfo.ErrorMessages, FOnIsCancelled::CreateRaw(&InCommand, &FPerforceSourceControlCommand::IsCanceled), InCommand.bConnectionDropped);
+		InCommand.bCommandSuccessful = Connection.RunCommand(TEXT("add"), Parameters, Records, InCommand.ResultInfo, FOnIsCancelled::CreateRaw(&InCommand, &FPerforceSourceControlCommand::IsCanceled), InCommand.bConnectionDropped);
 		ParseRecordSetForState(Records, OutResults);
 	}
 	return InCommand.bCommandSuccessful;
@@ -898,7 +898,7 @@ bool FPerforceDeleteWorker::Execute(FPerforceSourceControlCommand& InCommand)
 		Parameters.Append(InCommand.Files);
 
 		FP4RecordSet Records;
-		InCommand.bCommandSuccessful = Connection.RunCommand(TEXT("delete"), Parameters, Records, InCommand.ResultInfo.ErrorMessages, FOnIsCancelled::CreateRaw(&InCommand, &FPerforceSourceControlCommand::IsCanceled), InCommand.bConnectionDropped);
+		InCommand.bCommandSuccessful = Connection.RunCommand(TEXT("delete"), Parameters, Records, InCommand.ResultInfo, FOnIsCancelled::CreateRaw(&InCommand, &FPerforceSourceControlCommand::IsCanceled), InCommand.bConnectionDropped);
 		ParseRecordSetForState(Records, OutResults);
 	}
 	return InCommand.bCommandSuccessful;
@@ -960,7 +960,7 @@ bool FPerforceRevertWorker::Execute(FPerforceSourceControlCommand& InCommand)
 		}
 
 		FP4RecordSet Records;
-		InCommand.bCommandSuccessful = Connection.RunCommand(TEXT("revert"), Parameters, Records, InCommand.ResultInfo.ErrorMessages, FOnIsCancelled::CreateRaw(&InCommand, &FPerforceSourceControlCommand::IsCanceled), InCommand.bConnectionDropped);
+		InCommand.bCommandSuccessful = Connection.RunCommand(TEXT("revert"), Parameters, Records, InCommand.ResultInfo, FOnIsCancelled::CreateRaw(&InCommand, &FPerforceSourceControlCommand::IsCanceled), InCommand.bConnectionDropped);
 		ParseRecordSetForState(Records, OutResults);
 
 		for (TMap<FString, EPerforceState::Type>::TConstIterator It(OutResults); It; ++It)
@@ -1079,7 +1079,7 @@ bool FPerforceSyncWorker::Execute(FPerforceSourceControlCommand& InCommand)
 		}
 
 		FP4RecordSet Records;
-		InCommand.bCommandSuccessful = Connection.RunCommand(TEXT("sync"), Parameters, Records, InCommand.ResultInfo.ErrorMessages, FOnIsCancelled::CreateRaw(&InCommand, &FPerforceSourceControlCommand::IsCanceled), InCommand.bConnectionDropped);
+		InCommand.bCommandSuccessful = Connection.RunCommand(TEXT("sync"), Parameters, Records, InCommand.ResultInfo, FOnIsCancelled::CreateRaw(&InCommand, &FPerforceSourceControlCommand::IsCanceled), InCommand.bConnectionDropped);
 		ParseSyncResults(Records, OutResults);
 
 		RemoveRedundantErrors(InCommand, TEXT("file(s) up-to-date"));
@@ -1785,7 +1785,7 @@ static bool GetFileHistory(	FPerforceSourceControlProvider& SCCProvider, FPerfor
 	//limit to last 100 changes
 	Parameters.Add(TEXT("-m 100"));
 	Parameters.Append(InFiles);
-	InCommand.bCommandSuccessful &= Connection.RunCommand(TEXT("filelog"), Parameters, Records, InCommand.ResultInfo.ErrorMessages, FOnIsCancelled::CreateRaw(&InCommand, &FPerforceSourceControlCommand::IsCanceled), InCommand.bConnectionDropped);
+	InCommand.bCommandSuccessful &= Connection.RunCommand(TEXT("filelog"), Parameters, Records, InCommand.ResultInfo, FOnIsCancelled::CreateRaw(&InCommand, &FPerforceSourceControlCommand::IsCanceled), InCommand.bConnectionDropped);
 	ParseHistoryResults(SCCProvider, Records, OutStates, OutHistory);
 	RemoveRedundantErrors(InCommand, TEXT(" - no such file(s)."));
 	RemoveRedundantErrors(InCommand, TEXT(" - file(s) not on client"));
@@ -1902,12 +1902,12 @@ bool FPerforceUpdateStatusWorker::Execute(FPerforceSourceControlCommand& InComma
 				RevisionParameters.Insert(TEXT("-Of"), 0);
 
 				FP4RecordSet RevisionRecords;
-				InCommand.bCommandSuccessful &= Connection.RunCommand(TEXT("fstat"), RevisionParameters, RevisionRecords, InCommand.ResultInfo.ErrorMessages, FOnIsCancelled::CreateRaw(&InCommand, &FPerforceSourceControlCommand::IsCanceled), InCommand.bConnectionDropped);
+				InCommand.bCommandSuccessful &= Connection.RunCommand(TEXT("fstat"), RevisionParameters, RevisionRecords, InCommand.ResultInfo, FOnIsCancelled::CreateRaw(&InCommand, &FPerforceSourceControlCommand::IsCanceled), InCommand.bConnectionDropped);
 				ParseBranchModificationResults(RevisionRecords, InCommand.ResultInfo.ErrorMessages, ContentRoot, BranchModifications);
 			}
 
 			FP4RecordSet Records;
-			InCommand.bCommandSuccessful &= Connection.RunCommand(TEXT("fstat"), Parameters, Records, InCommand.ResultInfo.ErrorMessages, FOnIsCancelled::CreateRaw(&InCommand, &FPerforceSourceControlCommand::IsCanceled), InCommand.bConnectionDropped);
+			InCommand.bCommandSuccessful &= Connection.RunCommand(TEXT("fstat"), Parameters, Records, InCommand.ResultInfo, FOnIsCancelled::CreateRaw(&InCommand, &FPerforceSourceControlCommand::IsCanceled), InCommand.bConnectionDropped);
 			ParseUpdateStatusResults(Records, InCommand.ResultInfo.ErrorMessages, OutStates, ContentRoot, BranchModifications);
 			RemoveRedundantErrors(InCommand, TEXT(" - no such file(s)."), false);
 			RemoveRedundantErrors(InCommand, TEXT("' is not under client's root '"));
@@ -1930,7 +1930,7 @@ bool FPerforceUpdateStatusWorker::Execute(FPerforceSourceControlCommand& InComma
 			TArray<FString> Parameters = InCommand.Files;
 			Parameters.Add(FileQuery);
 			FP4RecordSet Records;
-			Connection.RunCommand(TEXT("opened"), Parameters, Records, InCommand.ResultInfo.ErrorMessages, FOnIsCancelled::CreateRaw(&InCommand, &FPerforceSourceControlCommand::IsCanceled), InCommand.bConnectionDropped);
+			Connection.RunCommand(TEXT("opened"), Parameters, Records, InCommand.ResultInfo, FOnIsCancelled::CreateRaw(&InCommand, &FPerforceSourceControlCommand::IsCanceled), InCommand.bConnectionDropped);
 			InCommand.bCommandSuccessful &= (InCommand.ResultInfo.ErrorMessages.Num() == 0);
 			ParseOpenedResults(Records, ANSI_TO_TCHAR(Connection.P4Client.GetClient().Text()), Connection.ClientRoot, OutStateMap);
 			RemoveRedundantErrors(InCommand, TEXT(" - no such file(s)."));
@@ -1964,7 +1964,7 @@ bool FPerforceUpdateStatusWorker::Execute(FPerforceSourceControlCommand& InComma
 
 				Parameters.Add(MoveTemp(File));
 			}
-			InCommand.bCommandSuccessful &= Connection.RunCommand(TEXT("diff"), Parameters, Records, InCommand.ResultInfo.ErrorMessages, FOnIsCancelled::CreateRaw(&InCommand, &FPerforceSourceControlCommand::IsCanceled), InCommand.bConnectionDropped);
+			InCommand.bCommandSuccessful &= Connection.RunCommand(TEXT("diff"), Parameters, Records, InCommand.ResultInfo, FOnIsCancelled::CreateRaw(&InCommand, &FPerforceSourceControlCommand::IsCanceled), InCommand.bConnectionDropped);
 
 			bool bMoveErrorsToInfo = !Operation->ShouldBeQuiet();
 
@@ -2137,7 +2137,7 @@ bool FPerforceGetWorkspacesWorker::Execute(FPerforceSourceControlCommand& InComm
 	{
 		FPerforceConnection& Connection = ScopedConnection.GetConnection();
 		TArray<FString> ClientSpecList;
-		InCommand.bCommandSuccessful = Connection.GetWorkspaceList(InCommand.ConnectionInfo, FOnIsCancelled::CreateRaw(&InCommand, &FPerforceSourceControlCommand::IsCanceled), ClientSpecList, InCommand.ResultInfo.ErrorMessages);
+		InCommand.bCommandSuccessful = Connection.GetWorkspaceList(InCommand.ConnectionInfo, FOnIsCancelled::CreateRaw(&InCommand, &FPerforceSourceControlCommand::IsCanceled), ClientSpecList, InCommand.ResultInfo);
 
 		check(InCommand.Operation->GetName() == GetName());
 		TSharedRef<FGetWorkspaces, ESPMode::ThreadSafe> Operation = StaticCastSharedRef<FGetWorkspaces>(InCommand.Operation);
@@ -2180,7 +2180,7 @@ static bool GetOpenedFilesInChangelist(
 	}
 
 	FP4RecordSet Records;
-	Connection.RunCommand(TEXT("opened"), Parameters, Records, InCommand.ResultInfo.ErrorMessages, FOnIsCancelled::CreateRaw(&InCommand, &FPerforceSourceControlCommand::IsCanceled), InCommand.bConnectionDropped);
+	Connection.RunCommand(TEXT("opened"), Parameters, Records, InCommand.ResultInfo, FOnIsCancelled::CreateRaw(&InCommand, &FPerforceSourceControlCommand::IsCanceled), InCommand.bConnectionDropped);
 	InCommand.bCommandSuccessful &= (InCommand.ResultInfo.ErrorMessages.Num() == 0);
 
 	if (InCommand.bCommandSuccessful)
@@ -2221,9 +2221,9 @@ static bool GetDepotFileToLocalFileMap(FPerforceConnection& Connection, FPerforc
 	}
 
 	FP4RecordSet Records;
-	Connection.RunCommand(TEXT("where"), Parameters, Records, InCommand.ResultInfo.ErrorMessages, FOnIsCancelled::CreateRaw(&InCommand, &FPerforceSourceControlCommand::IsCanceled), InCommand.bConnectionDropped);
+	Connection.RunCommand(TEXT("where"), Parameters, Records, InCommand.ResultInfo, FOnIsCancelled::CreateRaw(&InCommand, &FPerforceSourceControlCommand::IsCanceled), InCommand.bConnectionDropped);
 
-	if (InCommand.ResultInfo.ErrorMessages.Num() == 0)
+	if (!InCommand.ResultInfo.HasErrors())
 	{
 		ParseWhereResults(Records, OutDepotToLocalMap);
 		return true;
@@ -2266,7 +2266,7 @@ bool FPerforceGetPendingChangelistsWorker::Execute(FPerforceSourceControlCommand
 			Parameters.Add(InCommand.ConnectionInfo.Workspace);			// <workspace>
 
 			FP4RecordSet Records;
-			Connection.RunCommand(TEXT("changes"), Parameters, Records, InCommand.ResultInfo.ErrorMessages, FOnIsCancelled::CreateRaw(&InCommand, &FPerforceSourceControlCommand::IsCanceled), InCommand.bConnectionDropped);
+			Connection.RunCommand(TEXT("changes"), Parameters, Records, InCommand.ResultInfo, FOnIsCancelled::CreateRaw(&InCommand, &FPerforceSourceControlCommand::IsCanceled), InCommand.bConnectionDropped);
 			InCommand.bCommandSuccessful &= (InCommand.ResultInfo.ErrorMessages.Num() == 0);
 
 			ParseChangelistsResults(Records, OutChangelistsStates);
@@ -2353,7 +2353,7 @@ bool FPerforceGetPendingChangelistsWorker::Execute(FPerforceSourceControlCommand
 				Parameters.Add(ChangelistState.Changelist.ToString());
 
 				FP4RecordSet Records;
-				InCommand.bCommandSuccessful = Connection.RunCommand(TEXT("describe"), Parameters, Records, InCommand.ResultInfo.ErrorMessages, FOnIsCancelled::CreateRaw(&InCommand, &FPerforceSourceControlCommand::IsCanceled), InCommand.bConnectionDropped);
+				InCommand.bCommandSuccessful = Connection.RunCommand(TEXT("describe"), Parameters, Records, InCommand.ResultInfo, FOnIsCancelled::CreateRaw(&InCommand, &FPerforceSourceControlCommand::IsCanceled), InCommand.bConnectionDropped);
 
 				if (InCommand.bCommandSuccessful)
 				{
@@ -2530,7 +2530,7 @@ bool FPerforceCopyWorker::Execute(FPerforceSourceControlCommand& InCommand)
 			Parameters.Append(InCommand.Files);
 			Parameters.Add(DestinationPath);
 
-			InCommand.bCommandSuccessful = Connection.RunCommand(TEXT("integrate"), Parameters, Records, InCommand.ResultInfo.ErrorMessages, FOnIsCancelled::CreateRaw(&InCommand, &FPerforceSourceControlCommand::IsCanceled), InCommand.bConnectionDropped);
+			InCommand.bCommandSuccessful = Connection.RunCommand(TEXT("integrate"), Parameters, Records, InCommand.ResultInfo, FOnIsCancelled::CreateRaw(&InCommand, &FPerforceSourceControlCommand::IsCanceled), InCommand.bConnectionDropped);
 
 			// We now need to do a p4 resolve.
 			// This is because when we copy a file in the Editor, we first make the copy on disk before attempting to branch. This causes a conflict in P4's eyes.
@@ -2541,7 +2541,7 @@ bool FPerforceCopyWorker::Execute(FPerforceSourceControlCommand& InCommand)
 				TArray<FString> ResolveParameters;
 				ResolveParameters.Add(TEXT("-ay"));	// 'accept yours'
 				ResolveParameters.Add(MoveTemp(DestinationPath));
-				InCommand.bCommandSuccessful = Connection.RunCommand(TEXT("resolve"), ResolveParameters, Records, InCommand.ResultInfo.ErrorMessages, FOnIsCancelled::CreateRaw(&InCommand, &FPerforceSourceControlCommand::IsCanceled), InCommand.bConnectionDropped);
+				InCommand.bCommandSuccessful = Connection.RunCommand(TEXT("resolve"), ResolveParameters, Records, InCommand.ResultInfo, FOnIsCancelled::CreateRaw(&InCommand, &FPerforceSourceControlCommand::IsCanceled), InCommand.bConnectionDropped);
 			}
 		}
 		else
@@ -2550,7 +2550,7 @@ bool FPerforceCopyWorker::Execute(FPerforceSourceControlCommand& InCommand)
 
 			check(Operation->CopyMethod == FCopy::ECopyMethod::Add);
 
-			InCommand.bCommandSuccessful = Connection.RunCommand(TEXT("add"), Parameters, Records, InCommand.ResultInfo.ErrorMessages, FOnIsCancelled::CreateRaw(&InCommand, &FPerforceSourceControlCommand::IsCanceled), InCommand.bConnectionDropped);
+			InCommand.bCommandSuccessful = Connection.RunCommand(TEXT("add"), Parameters, Records, InCommand.ResultInfo, FOnIsCancelled::CreateRaw(&InCommand, &FPerforceSourceControlCommand::IsCanceled), InCommand.bConnectionDropped);
 		}
 	}
 	return InCommand.bCommandSuccessful;
@@ -2584,7 +2584,7 @@ bool FPerforceResolveWorker::Execute(FPerforceSourceControlCommand& InCommand)
 		AppendChangelistParameter(GetSCCProvider(), Parameters);
 
 		FP4RecordSet Records;
-		InCommand.bCommandSuccessful = Connection.RunCommand(TEXT("resolve"), Parameters, Records, InCommand.ResultInfo.ErrorMessages, FOnIsCancelled::CreateRaw(&InCommand, &FPerforceSourceControlCommand::IsCanceled), InCommand.bConnectionDropped);
+		InCommand.bCommandSuccessful = Connection.RunCommand(TEXT("resolve"), Parameters, Records, InCommand.ResultInfo, FOnIsCancelled::CreateRaw(&InCommand, &FPerforceSourceControlCommand::IsCanceled), InCommand.bConnectionDropped);
 		if( InCommand.bCommandSuccessful )
 		{
 			UpdatedFiles = InCommand.Files;
@@ -2626,7 +2626,7 @@ bool FPerforceChangeStatusWorker::Execute(FPerforceSourceControlCommand& InComma
 		Parameters.Append(InCommand.Files);
 
 		FP4RecordSet Records;
-		InCommand.bCommandSuccessful = Connection.RunCommand(TEXT("cstat"), Parameters, Records, InCommand.ResultInfo.ErrorMessages, FOnIsCancelled::CreateRaw(&InCommand, &FPerforceSourceControlCommand::IsCanceled), InCommand.bConnectionDropped);
+		InCommand.bCommandSuccessful = Connection.RunCommand(TEXT("cstat"), Parameters, Records, InCommand.ResultInfo, FOnIsCancelled::CreateRaw(&InCommand, &FPerforceSourceControlCommand::IsCanceled), InCommand.bConnectionDropped);
 		if (InCommand.bCommandSuccessful)
 		{
 			TSharedRef<FPerforceSourceControlChangeStatusOperation, ESPMode::ThreadSafe> Operation = StaticCastSharedRef<FPerforceSourceControlChangeStatusOperation>(InCommand.Operation);
@@ -2683,7 +2683,7 @@ bool FPerforceNewChangelistWorker::Execute(FPerforceSourceControlCommand& InComm
 		check(InCommand.Operation->GetName() == GetName());
 		TSharedRef<FNewChangelist, ESPMode::ThreadSafe> Operation = StaticCastSharedRef<FNewChangelist>(InCommand.Operation);
 
-		int32 ChangeList = Connection.CreatePendingChangelist(Operation->GetDescription(), TArray<FString>(), FOnIsCancelled::CreateRaw(&InCommand, &FPerforceSourceControlCommand::IsCanceled), InCommand.ResultInfo.ErrorMessages);
+		int32 ChangeList = Connection.CreatePendingChangelist(Operation->GetDescription(), TArray<FString>(), FOnIsCancelled::CreateRaw(&InCommand, &FPerforceSourceControlCommand::IsCanceled), InCommand.ResultInfo);
 
 		InCommand.bCommandSuccessful = (ChangeList > 0);
 		
@@ -2707,7 +2707,7 @@ bool FPerforceNewChangelistWorker::Execute(FPerforceSourceControlCommand& InComm
 					TArray<FString> ChangeParams;
 					ChangeParams.Add(TEXT("-d"));
 					ChangeParams.Add(NewChangelist.ToString());
-					Connection.RunCommand(TEXT("change"), ChangeParams, Records, InCommand.ResultInfo.ErrorMessages, FOnIsCancelled::CreateRaw(&InCommand, &FPerforceSourceControlCommand::IsCanceled), InCommand.bConnectionDropped);
+					Connection.RunCommand(TEXT("change"), ChangeParams, Records, InCommand.ResultInfo, FOnIsCancelled::CreateRaw(&InCommand, &FPerforceSourceControlCommand::IsCanceled), InCommand.bConnectionDropped);
 
 					InCommand.bCommandSuccessful = false;
 				}
@@ -2782,7 +2782,7 @@ bool FPerforceDeleteChangelistWorker::Execute(FPerforceSourceControlCommand& InC
 		Params.Add(TEXT("-d"));
 		Params.Add(InCommand.Changelist.ToString());
 		// Command will fail if changelist is not empty
-		Connection.RunCommand(TEXT("change"), Params, Records, InCommand.ResultInfo.ErrorMessages, FOnIsCancelled::CreateRaw(&InCommand, &FPerforceSourceControlCommand::IsCanceled), InCommand.bConnectionDropped);
+		Connection.RunCommand(TEXT("change"), Params, Records, InCommand.ResultInfo, FOnIsCancelled::CreateRaw(&InCommand, &FPerforceSourceControlCommand::IsCanceled), InCommand.bConnectionDropped);
 		// The normal parsing of the records here will show that it failed, but there's no record on a deleted changelist
 		InCommand.bCommandSuccessful = (InCommand.ResultInfo.ErrorMessages.Num() == 0);
 		
@@ -2831,11 +2831,11 @@ bool FPerforceEditChangelistWorker::Execute(FPerforceSourceControlCommand& InCom
 
 		if (InCommand.Changelist.IsDefault())
 		{
-			ChangelistNumber = Connection.CreatePendingChangelist(Operation->GetDescription(), InCommand.Files, FOnIsCancelled::CreateRaw(&InCommand, &FPerforceSourceControlCommand::IsCanceled), InCommand.ResultInfo.ErrorMessages);
+			ChangelistNumber = Connection.CreatePendingChangelist(Operation->GetDescription(), InCommand.Files, FOnIsCancelled::CreateRaw(&InCommand, &FPerforceSourceControlCommand::IsCanceled), InCommand.ResultInfo);
 		}
 		else
 		{
-			ChangelistNumber = Connection.EditPendingChangelist(Operation->GetDescription(), InCommand.Changelist.ToInt(), FOnIsCancelled::CreateRaw(&InCommand, &FPerforceSourceControlCommand::IsCanceled), InCommand.ResultInfo.ErrorMessages);
+			ChangelistNumber = Connection.EditPendingChangelist(Operation->GetDescription(), InCommand.Changelist.ToInt(), FOnIsCancelled::CreateRaw(&InCommand, &FPerforceSourceControlCommand::IsCanceled), InCommand.ResultInfo);
 		}
 
 		InCommand.bCommandSuccessful = (ChangelistNumber == InCommand.Changelist.ToInt() || (ChangelistNumber >= 0 && InCommand.Changelist.IsDefault()));
@@ -2888,7 +2888,7 @@ bool FPerforceRevertUnchangedWorker::Execute(FPerforceSourceControlCommand& InCo
 		}	
 
 		FP4RecordSet Records;
-		InCommand.bCommandSuccessful = Connection.RunCommand(TEXT("revert"), Parameters, Records, InCommand.ResultInfo.ErrorMessages, FOnIsCancelled::CreateRaw(&InCommand, &FPerforceSourceControlCommand::IsCanceled), InCommand.bConnectionDropped);
+		InCommand.bCommandSuccessful = Connection.RunCommand(TEXT("revert"), Parameters, Records, InCommand.ResultInfo, FOnIsCancelled::CreateRaw(&InCommand, &FPerforceSourceControlCommand::IsCanceled), InCommand.bConnectionDropped);
 		ParseRecordSetForState(Records, OutResults);
 		ChangelistToUpdate = InCommand.Changelist;
 	}
@@ -2989,7 +2989,7 @@ bool FPerforceShelveWorker::Execute(FPerforceSourceControlCommand& InCommand)
 					});
 			}
 
-			int32 NewChangeList = Connection.CreatePendingChangelist(Operation->GetDescription(), TArray<FString>(), FOnIsCancelled::CreateRaw(&InCommand, &FPerforceSourceControlCommand::IsCanceled), InCommand.ResultInfo.ErrorMessages);
+			int32 NewChangeList = Connection.CreatePendingChangelist(Operation->GetDescription(), TArray<FString>(), FOnIsCancelled::CreateRaw(&InCommand, &FPerforceSourceControlCommand::IsCanceled), InCommand.ResultInfo);
 			if (NewChangeList > 0)
 			{
 				Changelist = FPerforceSourceControlChangelist(NewChangeList);
@@ -3016,7 +3016,7 @@ bool FPerforceShelveWorker::Execute(FPerforceSourceControlCommand& InCommand)
 				Parameters.Append(InCommand.Files);
 			}
 
-			InCommand.bCommandSuccessful = Connection.RunCommand(TEXT("shelve"), Parameters, Records, InCommand.ResultInfo.ErrorMessages, FOnIsCancelled::CreateRaw(&InCommand, &FPerforceSourceControlCommand::IsCanceled), InCommand.bConnectionDropped);
+			InCommand.bCommandSuccessful = Connection.RunCommand(TEXT("shelve"), Parameters, Records, InCommand.ResultInfo, FOnIsCancelled::CreateRaw(&InCommand, &FPerforceSourceControlCommand::IsCanceled), InCommand.bConnectionDropped);
 		}
 
 		if (InCommand.bCommandSuccessful)
@@ -3043,7 +3043,7 @@ bool FPerforceShelveWorker::Execute(FPerforceSourceControlCommand& InCommand)
 				TArray<FString> ChangeParams;
 				ChangeParams.Add(TEXT("-d"));
 				ChangeParams.Add(Changelist.ToString());
-				Connection.RunCommand(TEXT("change"), ChangeParams, Records, InCommand.ResultInfo.ErrorMessages, FOnIsCancelled::CreateRaw(&InCommand, &FPerforceSourceControlCommand::IsCanceled), InCommand.bConnectionDropped);
+				Connection.RunCommand(TEXT("change"), ChangeParams, Records, InCommand.ResultInfo, FOnIsCancelled::CreateRaw(&InCommand, &FPerforceSourceControlCommand::IsCanceled), InCommand.bConnectionDropped);
 			}
 		}
 	}
@@ -3111,7 +3111,7 @@ bool FPerforceDeleteShelveWorker::Execute(FPerforceSourceControlCommand& InComma
 		}
 
 		FP4RecordSet Records;
-		Connection.RunCommand(TEXT("shelve"), Parameters, Records, InCommand.ResultInfo.ErrorMessages, FOnIsCancelled::CreateRaw(&InCommand, &FPerforceSourceControlCommand::IsCanceled), InCommand.bConnectionDropped);
+		Connection.RunCommand(TEXT("shelve"), Parameters, Records, InCommand.ResultInfo, FOnIsCancelled::CreateRaw(&InCommand, &FPerforceSourceControlCommand::IsCanceled), InCommand.bConnectionDropped);
 		InCommand.bCommandSuccessful = (InCommand.ResultInfo.ErrorMessages.Num() == 0);
 
 		if (InCommand.bCommandSuccessful)
@@ -3179,7 +3179,7 @@ bool FPerforceUnshelveWorker::Execute(FPerforceSourceControlCommand& InCommand)
 
 		FP4RecordSet Records;
 		// Note: unshelve can succeed partially.
-		InCommand.bCommandSuccessful = Connection.RunCommand(TEXT("unshelve"), Parameters, Records, InCommand.ResultInfo.ErrorMessages, FOnIsCancelled::CreateRaw(&InCommand, &FPerforceSourceControlCommand::IsCanceled), InCommand.bConnectionDropped);
+		InCommand.bCommandSuccessful = Connection.RunCommand(TEXT("unshelve"), Parameters, Records, InCommand.ResultInfo, FOnIsCancelled::CreateRaw(&InCommand, &FPerforceSourceControlCommand::IsCanceled), InCommand.bConnectionDropped);
 		
 		if (InCommand.bCommandSuccessful && Records.Num() > 0)
 		{
@@ -3253,7 +3253,7 @@ bool FPerforceDownloadFileWorker::Execute(FPerforceSourceControlCommand& InComma
 			TArray<FSharedBuffer> FilesData;
 			const ERunCommandFlags Flags = Operation->ShouldLogToStdOutput() ? ERunCommandFlags::Default : ERunCommandFlags::DisableCommandLogging;
 
-			InCommand.bCommandSuccessful = Connection.RunCommand(	TEXT("print"), Parameters, Records, &FilesData, InCommand.ResultInfo.ErrorMessages,
+			InCommand.bCommandSuccessful = Connection.RunCommand(	TEXT("print"), Parameters, Records, &FilesData, InCommand.ResultInfo,
 																	FOnIsCancelled::CreateRaw(&InCommand, &FPerforceSourceControlCommand::IsCanceled), 
 																	InCommand.bConnectionDropped, Flags);
 			
@@ -3297,7 +3297,7 @@ bool FPerforceDownloadFileWorker::Execute(FPerforceSourceControlCommand& InComma
 
 				const ERunCommandFlags Flags = Operation->ShouldLogToStdOutput() ? ERunCommandFlags::Default : ERunCommandFlags::DisableCommandLogging;
 
-				InCommand.bCommandSuccessful = Connection.RunCommand(	TEXT("print"), Parameters, Records, nullptr, InCommand.ResultInfo.ErrorMessages, 
+				InCommand.bCommandSuccessful = Connection.RunCommand(	TEXT("print"), Parameters, Records, nullptr, InCommand.ResultInfo, 
 																		FOnIsCancelled::CreateRaw(&InCommand, &FPerforceSourceControlCommand::IsCanceled), 
 																		InCommand.bConnectionDropped, Flags);
 
@@ -3401,7 +3401,7 @@ bool FPerforceCreateWorkspaceWorker::Execute(FPerforceSourceControlCommand& InCo
 
 		InCommand.bCommandSuccessful = Connection.CreateWorkspace(	ClientDesc,
 																	FOnIsCancelled::CreateRaw(&InCommand, &FPerforceSourceControlCommand::IsCanceled),
-																	InCommand.ResultInfo.ErrorMessages);
+																	InCommand.ResultInfo);
 	}
 	
 	return InCommand.bCommandSuccessful;
@@ -3466,7 +3466,7 @@ bool FPerforceDeleteWorkspaceWorker::Execute(FPerforceSourceControlCommand& InCo
 								Parameters, 
 								Records, 
 								nullptr,
-								InCommand.ResultInfo.ErrorMessages, 
+								InCommand.ResultInfo, 
 								FOnIsCancelled::CreateRaw(&InCommand, &FPerforceSourceControlCommand::IsCanceled), 
 								InCommand.bConnectionDropped,
 								RunFlags);
@@ -3506,7 +3506,7 @@ bool FPerforceGetChangelistDetailsWorker::Execute(FPerforceSourceControlCommand&
 		Parameters.Add(Operation->GetChangelistNumber());
 
 		FP4RecordSet Records;
-		InCommand.bCommandSuccessful = Connection.RunCommand(TEXT("describe"), Parameters, Records, InCommand.ResultInfo.ErrorMessages, FOnIsCancelled::CreateRaw(&InCommand, &FPerforceSourceControlCommand::IsCanceled), InCommand.bConnectionDropped);
+		InCommand.bCommandSuccessful = Connection.RunCommand(TEXT("describe"), Parameters, Records, InCommand.ResultInfo, FOnIsCancelled::CreateRaw(&InCommand, &FPerforceSourceControlCommand::IsCanceled), InCommand.bConnectionDropped);
 
 		if (InCommand.bCommandSuccessful)
 		{
diff --git a/Engine/Plugins/Developer/PerforceSourceControl/Source/PerforceSourceControl/Private/PerforceSourceControlProvider.cpp b/Engine/Plugins/Developer/PerforceSourceControl/Source/PerforceSourceControl/Private/PerforceSourceControlProvider.cpp
--- a/Engine/Plugins/Developer/PerforceSourceControl/Source/PerforceSourceControl/Private/PerforceSourceControlProvider.cpp
+++ b/Engine/Plugins/Developer/PerforceSourceControl/Source/PerforceSourceControl/Private/PerforceSourceControlProvider.cpp
@@ -313,13 +313,6 @@ ISourceControlProvider::FInitResult FPerforceSourceControlProvider::ParseCommand
 	return Result;
 }
 
-void FPerforceSourceControlProvider::GetWorkspaceList(const FPerforceConnectionInfo& InConnectionInfo, TArray<FString>& OutWorkspaceList, TArray<FText>& OutErrorMessages)
-{
-	//attempt to ask perforce for a list of client specs that belong to this user
-	FPerforceConnection Connection(InConnectionInfo, *this);
-	Connection.GetWorkspaceList(InConnectionInfo, FOnIsCancelled(), OutWorkspaceList, OutErrorMessages);
-}
-
 const FString& FPerforceSourceControlProvider::GetTicket() const
 {
 	return Ticket;
@@ -737,21 +730,20 @@ TArray< TSharedRef<ISourceControlLabel> > FPerforceSourceControlProvider::GetLab
 		FPerforceConnection& Connection = ScopedConnection.GetConnection();
 		FP4RecordSet Records;
 		TArray<FString> Parameters;
-		TArray<FText> ErrorMessages;
+		FSourceControlResultInfo ResultInfo;
 		Parameters.Add(TEXT("-E"));
 		Parameters.Add(InMatchingSpec);
 		bool bConnectionDropped = false;
-		if(Connection.RunCommand(TEXT("labels"), Parameters, Records, ErrorMessages, FOnIsCancelled(), bConnectionDropped))
+		if(Connection.RunCommand(TEXT("labels"), Parameters, Records, ResultInfo, FOnIsCancelled(), bConnectionDropped))
 		{
 			// const_cast to avoid changing the ISourceControlProvider API as it is hard to deprecate without causing derived types to give compiler errors.
 			ParseGetLabelsResults(*const_cast<FPerforceSourceControlProvider*>(this), Records, Labels);
 		}
 		else
 		{
-			// output errors if any
-			for (int32 ErrorIndex = 0; ErrorIndex < ErrorMessages.Num(); ++ErrorIndex)
+			for (const FText& ErrorMsg : ResultInfo.ErrorMessages)
 			{
-				FMessageLog("SourceControl").Warning(FText::Format(LOCTEXT("GetLabelsErrorFormat", "GetLabels Warning: {0}"), ErrorMessages[ErrorIndex]));
+				FMessageLog("SourceControl").Warning(FText::Format(LOCTEXT("GetLabelsErrorFormat", "GetLabels Warning: {0}"), ErrorMsg));
 			}
 		}
 	}
@@ -1055,7 +1047,7 @@ bool FPerforceSourceControlProvider::QueryStateBranchConfig(const FString& Confi
 		FPerforceConnection& Connection = ScopedConnection.GetConnection();
 		FP4RecordSet Records;
 		TArray<FString> Parameters;
-		TArray<FText> ErrorMessages;
+		FSourceControlResultInfo ResultInfo;
 		Parameters.Add(TEXT("-o"));
 		Parameters.Add(*ConfigDest);
 		Parameters.Add(*ConfigSrc);
@@ -1063,7 +1055,7 @@ bool FPerforceSourceControlProvider::QueryStateBranchConfig(const FString& Confi
 		FText GeneralErrorMessage = LOCTEXT("StatusBranchConfigGeneralFailure", "Unable to retrieve status branch configuration from depot");
 
 		bool bConnectionDropped = false;
-		if (Connection.RunCommand(TEXT("print"), Parameters, Records, ErrorMessages, FOnIsCancelled(), bConnectionDropped))
+		if (Connection.RunCommand(TEXT("print"), Parameters, Records, ResultInfo, FOnIsCancelled(), bConnectionDropped))
 		{
 			if (Records.Num() < 1 || Records[0][TEXT("depotFile")] != ConfigSrc)
 			{
@@ -1075,10 +1067,9 @@ bool FPerforceSourceControlProvider::QueryStateBranchConfig(const FString& Confi
 		{
 			FMessageLog("SourceControl").Error(GeneralErrorMessage);
 
-			// output specific errors if any
-			for (int32 ErrorIndex = 0; ErrorIndex < ErrorMessages.Num(); ++ErrorIndex)
+			for (const FText& ErrorMsg : ResultInfo.ErrorMessages)
 			{
-				FMessageLog("SourceControl").Error(ErrorMessages[ErrorIndex]);
+				FMessageLog("SourceControl").Error(ErrorMsg);
 			}
 
 			return false;
diff --git a/Engine/Plugins/Developer/PerforceSourceControl/Source/PerforceSourceControl/Private/PerforceSourceControlProvider.h b/Engine/Plugins/Developer/PerforceSourceControl/Source/PerforceSourceControl/Private/PerforceSourceControlProvider.h
--- a/Engine/Plugins/Developer/PerforceSourceControl/Source/PerforceSourceControl/Private/PerforceSourceControlProvider.h
+++ b/Engine/Plugins/Developer/PerforceSourceControl/Source/PerforceSourceControl/Private/PerforceSourceControlProvider.h
@@ -65,15 +65,6 @@ public:
 
 	using ISourceControlProvider::Execute;
 
-	/**
-	 * Gets a list of client spec names from the source control provider
-	 *
-	 * @param	InConnectionInfo	Credentials for connection
-	 * @param	OutWorkspaceList	List of client spec name strings
-	 * @param	OutErrorMessages	List of any error messages that may have occurred
-	 */
-	void GetWorkspaceList(const struct FPerforceConnectionInfo& InConnectionInfo, TArray<FString>& OutWorkspaceList, TArray<FText>& OutErrorMessages);
-
 	/** Get the P4 ticket we will use for connections */
 	const FString& GetTicket() const;
 
diff --git a/Engine/Plugins/Developer/PerforceSourceControl/Source/PerforceSourceControl/Private/PerforceSourceControlRevision.cpp b/Engine/Plugins/Developer/PerforceSourceControl/Source/PerforceSourceControl/Private/PerforceSourceControlRevision.cpp
--- a/Engine/Plugins/Developer/PerforceSourceControl/Source/PerforceSourceControl/Private/PerforceSourceControlRevision.cpp
+++ b/Engine/Plugins/Developer/PerforceSourceControl/Source/PerforceSourceControl/Private/PerforceSourceControlRevision.cpp
@@ -45,7 +45,7 @@ bool FPerforceSourceControlRevision::Get(FString& InOutFilename, FPerforceConnec
 
 	FP4RecordSet Records;
 	bool bConnectionDropped;
-	TArray<FText> ErrorMessages;
+	FSourceControlResultInfo ResultInfo;
 	TArray<FString> Parameters;
 
 	// Suppress the one-line file header normally added by Perforce.
@@ -83,7 +83,7 @@ bool FPerforceSourceControlRevision::Get(FString& InOutFilename, FPerforceConnec
 	Parameters.Add(FString("-o") + AbsoluteFileName);
 	Parameters.Add(FileName + RevParam);
 
-	const bool bCmdResult = Connection.RunCommand(TEXT("print"), Parameters, Records, ErrorMessages, FOnIsCancelled(), bConnectionDropped);
+	const bool bCmdResult = Connection.RunCommand(TEXT("print"), Parameters, Records, ResultInfo, FOnIsCancelled(), bConnectionDropped);
 	if (bCmdResult)
 	{
 		InOutFilename = AbsoluteFileName;
@@ -91,9 +91,9 @@ bool FPerforceSourceControlRevision::Get(FString& InOutFilename, FPerforceConnec
 	}
 	else
 	{
-		for (auto Iter(ErrorMessages.CreateConstIterator()); Iter; Iter++)
+		for (const FText& ErrorMsg : ResultInfo.ErrorMessages)
 		{
-			FTSMessageLog("SourceControl").Error(FText::Format(LOCTEXT("PerforceSourceControlRevisionGetErrorFormat", "FPerforceSourceControlRevision::Get print Error: {0}"), *Iter));
+			FTSMessageLog("SourceControl").Error(FText::Format(LOCTEXT("PerforceSourceControlRevisionGetErrorFormat", "FPerforceSourceControlRevision::Get print Error: {0}"), ErrorMsg));
 		}
 
 		return false;
@@ -110,14 +110,14 @@ static FString GetUserFromChangelist(int32 ChangeNumber, FPerforceConnection& In
 {
 	FP4RecordSet Records;
 	bool bConnectionDropped;
-	TArray<FText> ErrorMessages;
+	FSourceControlResultInfo ResultInfo;
 	TArray<FString> Parameters;
 
 	// Only describe the basic changelist information, suppress output of the file diffs
 	Parameters.Add(TEXT("-s"));
 	Parameters.Add(FString::Printf(TEXT("%d"), ChangeNumber));
 
-	if(InConnection.RunCommand(TEXT("describe"), Parameters, Records, ErrorMessages, FOnIsCancelled(), bConnectionDropped))
+	if(InConnection.RunCommand(TEXT("describe"), Parameters, Records, ResultInfo, FOnIsCancelled(), bConnectionDropped))
 	{
 		if(Records.Num() > 0)
 		{
@@ -126,9 +126,9 @@ static FString GetUserFromChangelist(int32 ChangeNumber, FPerforceConnection& In
 	}
 	else
 	{
-		for(auto Iter(ErrorMessages.CreateConstIterator()); Iter; Iter++)
+		for (const FText& ErrorMsg : ResultInfo.ErrorMessages)
 		{
-			FMessageLog("SourceControl").Error(FText::Format(LOCTEXT("GetUserFromChangelistErrorFormat", "GetUserFromChangelist Error: {0}"), *Iter));
+			FMessageLog("SourceControl").Error(FText::Format(LOCTEXT("GetUserFromChangelistErrorFormat", "GetUserFromChangelist Error: {0}"), ErrorMsg));
 		}
 	}
 
@@ -180,26 +180,26 @@ bool FPerforceSourceControlRevision::GetAnnotated( TArray<FAnnotationLine>& OutL
 		FPerforceConnection& Connection = ScopedConnection.GetConnection();
 		FP4RecordSet Records;
 		bool bConnectionDropped;
-		TArray<FText> ErrorMessages;
+		FSourceControlResultInfo ResultInfo;
 		TArray<FString> Parameters;
 
 		Parameters.Add(TEXT("-q"));	// Suppress the one-line file header normally added by Perforce.
 		Parameters.Add(TEXT("-c"));	// Display change numbers rather than revision numbers
 		Parameters.Add(TEXT("-I"));	// Follow integrations
 
-		FString RevString = (RevisionNumber < 0) ? TEXT("head") : FString::Printf(TEXT("%d"), RevisionNumber);
+		const FString RevString = (RevisionNumber < 0) ? TEXT("head") : FString::Printf(TEXT("%d"), RevisionNumber);
 		Parameters.Add(FileName + TEXT("#") + RevString);
 
-		bCommandOK = Connection.RunCommand(TEXT("annotate"), Parameters, Records, ErrorMessages, FOnIsCancelled(), bConnectionDropped);
+		bCommandOK = Connection.RunCommand(TEXT("annotate"), Parameters, Records, ResultInfo, FOnIsCancelled(), bConnectionDropped);
 		if(bCommandOK)
 		{
 			ParseAnnotationResults(Records, OutLines, Connection);
 		}
 		else
 		{
-			for(auto Iter(ErrorMessages.CreateConstIterator()); Iter; Iter++)
+			for (const FText& ErrorMsg : ResultInfo.ErrorMessages)
 			{
-				FMessageLog("SourceControl").Error(FText::Format(LOCTEXT("GetAnnotatedErrorFormat", "GetAnnotated Error: {0}"), *Iter));
+				FMessageLog("SourceControl").Error(FText::Format(LOCTEXT("GetAnnotatedErrorFormat", "GetAnnotated Error: {0}"), ErrorMsg));
 			}
 		}
 	}
diff --git a/Engine/Source/Developer/SourceControl/Public/SourceControlResultInfo.h b/Engine/Source/Developer/SourceControl/Public/SourceControlResultInfo.h
--- a/Engine/Source/Developer/SourceControl/Public/SourceControlResultInfo.h
+++ b/Engine/Source/Developer/SourceControl/Public/SourceControlResultInfo.h
@@ -37,6 +37,11 @@ struct FSourceControlResultInfo
 
 	SOURCECONTROL_API bool DidConnectionFail() const;
 
+	bool HasErrors() const
+	{
+		return !ErrorMessages.IsEmpty();
+	}
+
 	/** Info and/or warning message storage */
 	TArray<FText> InfoMessages;
 
