Return files not being marked for add due to P4IGNORE as errors rather than info.

- The command will already be marked as a failure, but as p4 returns the message about the file(s) being ignored as INFO we will not print any error messages to the user and the failure can be confusing.
- Add ::RemoveRedundantInfo to allow us to remove info messages and move them to be errors. This is basically the same logic as ::RemoveRedundantErrors but going in the opposite direction.
- Check for the info message mentioning that a file failed to add as it is ignored (ignored file can't be added) and move that message to be an error.
-- The errors will be logged by the calling code so there is no need for us to do that here. In addition if the original calling code holds on to the operation pointer it will have access to the FSourceControlResultsInfo structure containing the error message if the calling code wishes to do it's own logging.

diff --git a/Engine/Plugins/Developer/PerforceSourceControl/Source/PerforceSourceControl/Private/PerforceSourceControlOperations.cpp b/Engine/Plugins/Developer/PerforceSourceControl/Source/PerforceSourceControl/Private/PerforceSourceControlOperations.cpp
--- a/Engine/Plugins/Developer/PerforceSourceControl/Source/PerforceSourceControl/Private/PerforceSourceControlOperations.cpp
+++ b/Engine/Plugins/Developer/PerforceSourceControl/Source/PerforceSourceControl/Private/PerforceSourceControlOperations.cpp
@@ -182,6 +182,32 @@ static void RemoveRedundantErrors(FPerforceSourceControlCommand& InCommand, cons
 	}
 }
 
+static void RemoveRedundantInfo(FPerforceSourceControlCommand& InCommand, const FString& InFilter, bool bMoveToError = true)
+{
+	bool bFoundRedundantInfo = false;
+
+	if (bMoveToError)
+	{
+		for (const FText& Message : InCommand.ResultInfo.InfoMessages)
+		{
+			// Perforce reports files that are already synced as errors, so copy any errors
+			// we get to the info list in this case
+			if (Message.ToString().Contains(InFilter))
+			{
+				InCommand.ResultInfo.ErrorMessages.Add(Message);
+			}
+		}
+	}
+
+	bFoundRedundantInfo = InCommand.ResultInfo.InfoMessages.RemoveAll(FRemoveRedundantErrors(InFilter)) > 0;
+
+	// If we found info that we moved to be an error we should mark the command as failed.
+	if (bFoundRedundantInfo && bMoveToError)
+	{
+		InCommand.bCommandSuccessful = false;
+	}
+}
+
 /**
  * Occasionally some error messages returned by the server can be a bit misleading depending on the context
  * of the original command. This utility allows us to replace an error that we feel is misleading with a
@@ -852,9 +878,11 @@ bool FPerforceMarkForAddWorker::Execute(FPerforceSourceControlCommand& InCommand
 
 		AppendChangelistParameter(GetSCCProvider(), Parameters);
 		Parameters.Append(InCommand.Files);
-
+	
 		InCommand.bCommandSuccessful = Connection.RunCommand(TEXT("add"), Parameters, Records, InCommand.ResultInfo, FOnIsCancelled::CreateRaw(&InCommand, &FPerforceSourceControlCommand::IsCanceled), InCommand.bConnectionDropped);
 		ParseRecordSetForState(Records, OutResults);
+
+		RemoveRedundantInfo(InCommand, TEXT("- ignored file can't be added."));
 	}
 	return InCommand.bCommandSuccessful;
 }
