Return a proper error message when a user tries to submit a changelist containing files not mapped to the current stream.

- Move info messages (warnings in this case) containing "- warning: file not mapped in stream" to be errors when submitting content."
- Cleaned up use of RemoveRedundantInfo to pass in the default parameter.

diff --git a/Engine/Plugins/Developer/PerforceSourceControl/Source/PerforceSourceControl/Private/PerforceSourceControlOperations.cpp b/Engine/Plugins/Developer/PerforceSourceControl/Source/PerforceSourceControl/Private/PerforceSourceControlOperations.cpp
--- a/Engine/Plugins/Developer/PerforceSourceControl/Source/PerforceSourceControl/Private/PerforceSourceControlOperations.cpp
+++ b/Engine/Plugins/Developer/PerforceSourceControl/Source/PerforceSourceControl/Private/PerforceSourceControlOperations.cpp
@@ -676,11 +676,14 @@ bool FPerforceCheckInWorker::Execute(FPerforceSourceControlCommand& InCommand)
 
 			InCommand.bCommandSuccessful = Connection.RunCommand(TEXT("submit"), SubmitParams, Records, InCommand.ResultInfo, FOnIsCancelled::CreateRaw(&InCommand, &FPerforceSourceControlCommand::IsCanceled), InCommand.bConnectionDropped);
 
-			if (InCommand.ResultInfo.ErrorMessages.Num() > 0)
+			if (InCommand.ResultInfo.HasErrors())
 			{
 				InCommand.bCommandSuccessful = false;
 			}
 
+			const bool bMoveToError = true;
+			RemoveRedundantInfo(InCommand, TEXT("- warning: file not mapped in stream"), bMoveToError);
+
 			if (InCommand.bCommandSuccessful)
 			{
 				// Remove any deleted files from status cache
@@ -882,7 +885,8 @@ bool FPerforceMarkForAddWorker::Execute(FPerforceSourceControlCommand& InCommand
 		InCommand.bCommandSuccessful = Connection.RunCommand(TEXT("add"), Parameters, Records, InCommand.ResultInfo, FOnIsCancelled::CreateRaw(&InCommand, &FPerforceSourceControlCommand::IsCanceled), InCommand.bConnectionDropped);
 		ParseRecordSetForState(Records, OutResults);
 
-		RemoveRedundantInfo(InCommand, TEXT("- ignored file can't be added."));
+		const bool bMoveToError = true;
+		RemoveRedundantInfo(InCommand, TEXT("- ignored file can't be added."), bMoveToError);
 	}
 	return InCommand.bCommandSuccessful;
 }
