Fix errors that were being logged incorrectly when updating changelist states for a user.

#rnx

- Updating the state of all changelists for a user currently does not take the users current workspace into account, meaning we will scan all changelists for a given user on the currently server. When we try to resolve the files in the changelist to local files for changelists not under the current workspace we will encountered perforce errors as the files are not in the current client spec view.
- To fix the immediate problem we should just remove those errors and changelists not under the current workspace will have the depot info about their files but not the local file path info, which is the current behavior but without the errors. In theory the error messages are correct but they are not that useful to the caller, especially as we do not provide context about which changelist is emitting each error.
-- Although the errors are raised inside of ::GetDepotFileToLocalFileMap we only remove this error in FPerforceGetPendingChangelistsWorker::Execute where we know we don't care about it. In some of the other places that we call ::GetDepotFileToLocalFileMap that error would indicate a proper problem so we should not change it there.

diff --git a/Engine/Plugins/Developer/PerforceSourceControl/Source/PerforceSourceControl/Private/PerforceSourceControlOperations.cpp b/Engine/Plugins/Developer/PerforceSourceControl/Source/PerforceSourceControl/Private/PerforceSourceControlOperations.cpp
--- a/Engine/Plugins/Developer/PerforceSourceControl/Source/PerforceSourceControl/Private/PerforceSourceControlOperations.cpp
+++ b/Engine/Plugins/Developer/PerforceSourceControl/Source/PerforceSourceControl/Private/PerforceSourceControlOperations.cpp
@@ -2394,6 +2394,16 @@ bool FPerforceGetPendingChangelistsWorker::Execute(FPerforceSourceControlCommand
 
 					TMap<FString, FString>& OutShelvedFileMap = OutCLShelvedFilesMap.Emplace_GetRef();
 					GetDepotFileToLocalFileMap(Connection, InCommand, OutShelvedStateMap, OutShelvedFileMap);
+
+					// Remove errors encountered in GetDepotFileToLocalFileMap when parsing changelists in other workspaces
+					// as we expect those files to not be in the client view
+					// TODO: The root cause of this issue is we iterate over all changelists but in general we only care
+					// about the changelists for the active workspace, we should
+					// a) Not try to resolve local files for changelists in other workspaces
+					// b) Provide better filtering so the user can request this operation only run on changelists for the
+					// active workspace.
+					const bool bMoveToInfo = false;
+					RemoveRedundantErrors(InCommand, TEXT(" - file(s) not in client view."), bMoveToInfo);
 				}
 			}
 		}
