Fix a crash when refreshing file states in perforce and the TypedElementFramework has not been initialized.

#rnx

- The crash was in a stand alone program that loads a very limited number of plugins which caused the TypedElementFramework to remain unloaded.
- UTypedElementRegistry::GetInstance is not guaranteed to return a valid pointer so we have to check the return value.

#changelist validated
#virtualized

diff --git a/Engine/Plugins/Developer/PerforceSourceControl/Source/PerforceSourceControl/Private/PerforceSourceControlOperations.cpp b/Engine/Plugins/Developer/PerforceSourceControl/Source/PerforceSourceControl/Private/PerforceSourceControlOperations.cpp
--- a/Engine/Plugins/Developer/PerforceSourceControl/Source/PerforceSourceControl/Private/PerforceSourceControlOperations.cpp
+++ b/Engine/Plugins/Developer/PerforceSourceControl/Source/PerforceSourceControl/Private/PerforceSourceControlOperations.cpp
@@ -2022,7 +2022,13 @@ bool FPerforceUpdateStatusWorker::UpdateStates() const
 		using namespace TypedElementQueryBuilder;
 		using DSI = ITypedElementDataStorageInterface;
 
-		DSI* DataStorage = UTypedElementRegistry::GetInstance()->GetMutableDataStorage();
+		UTypedElementRegistry* Registry = UTypedElementRegistry::GetInstance();
+		if (!Registry)
+		{
+			return;
+		}
+
+		DSI* DataStorage = Registry->GetMutableDataStorage();
 		if (!DataStorage)
 		{
 			return;
