Fix assert when failing to clobber a file.

#rnx

- The problem was introduced in my change CL 30208400 where I mistakenly believed that all messages should be routed through ClientUser::Message so added asserts to detect and classes derived from FP4ClientUser calling the lower level Info/Error methods. Server messages do indeed all go through ::Message but local client issues can call the lower level methods directly so my assumption was unsafe. To make matters worse the most common local error we have concerns clobbering files but I tend to run workspaces with "clobber" option rather than "noclobber" so I did not see the problem in testing.
- We once again implement ClientUser::OutputInfo and ClientUser::OutputError and store the messages in the correct array.
- The default implementation of ClientUser::HandleError calls ::OutputError so there is no point implementing this method but we also don't need to assert if it is called.

diff --git a/Engine/Plugins/Developer/PerforceSourceControl/Source/PerforceSourceControl/Private/PerforceConnection.cpp b/Engine/Plugins/Developer/PerforceSourceControl/Source/PerforceSourceControl/Private/PerforceConnection.cpp
--- a/Engine/Plugins/Developer/PerforceSourceControl/Source/PerforceSourceControl/Private/PerforceConnection.cpp
+++ b/Engine/Plugins/Developer/PerforceSourceControl/Source/PerforceSourceControl/Private/PerforceConnection.cpp
@@ -295,20 +295,21 @@ public:
 
 	virtual void OutputInfo(char Indent, const char* InInfo) override
 	{
-		// Should be caught by ::Message
-		checkNoEntry();
+		// We don't expect this to ever be called (info messages should come
+		// via ClientUser::Message) but implemented just to be safe.
+
+		ResultInfo.InfoMessages.Add(FText::FromString(FString(TO_TCHAR(InInfo, IsUnicodeServer()))));
 	}
 
 	virtual void OutputError(const char* errBuf) override
 	{
-		// Should be caught by ::Message
-		checkNoEntry();
-	}
+		// In general we expect errors to be passed to use via ClientUser::Message but some
+		// errors raised by the p4 cpp api can call ::HandleError or ::OutputError directly.
+		// Since the default implementation of ::HandleError calls ::OutputError we only need
+		// to implement this method to make sure we capture all of the errors being passed in
+		// this way.
 
-	virtual void HandleError(Error* InError) override
-	{
-		// Should be caught by ::Message
-		checkNoEntry();
+		ResultInfo.ErrorMessages.Add(FText::FromString(FString(TO_TCHAR(errBuf, IsUnicodeServer()))));
 	}
 
 	inline bool IsUnicodeServer() const
