Fix typo in log message when switching workspaces.

#rnx

diff --git a/Engine/Plugins/Developer/PerforceSourceControl/Source/PerforceSourceControl/Private/PerforceSourceControlProvider.cpp b/Engine/Plugins/Developer/PerforceSourceControl/Source/PerforceSourceControl/Private/PerforceSourceControlProvider.cpp
--- a/Engine/Plugins/Developer/PerforceSourceControl/Source/PerforceSourceControl/Private/PerforceSourceControlProvider.cpp
+++ b/Engine/Plugins/Developer/PerforceSourceControl/Source/PerforceSourceControl/Private/PerforceSourceControlProvider.cpp
@@ -910,7 +910,7 @@ ECommandResult::Type FPerforceSourceControlProvider::SwitchWorkspace(FStringView
 		P4Settings.SetWorkspace(WorkspaceName);
 	}
 
-	UE_LOG(LogSourceControl, Log, TEXT("Switched workspaces from '%s' to '%s%%'"), *OldWorkspaceName, *WorkspaceName);
+	UE_LOG(LogSourceControl, Log, TEXT("Switched workspaces from '%s' to '%s'"), *OldWorkspaceName, *WorkspaceName);
 
 	return ECommandResult::Succeeded;
 }
