Extend FSourceControlInitSettings to allow the caller to control if a provider uses values parsed from the commandline or not when initializing.

#rnx

- We default to the old behavior, so that the default way of creating a source control provider (no FSourceControlInitSettings) will read all of the commandline parameters available and the newer way (creating a provider with a FSourceControlInitSettings) will ignore the commandline.

diff --git a/Engine/Plugins/Developer/PerforceSourceControl/Source/PerforceSourceControl/Private/PerforceSourceControlProvider.cpp b/Engine/Plugins/Developer/PerforceSourceControl/Source/PerforceSourceControl/Private/PerforceSourceControlProvider.cpp
--- a/Engine/Plugins/Developer/PerforceSourceControl/Source/PerforceSourceControl/Private/PerforceSourceControlProvider.cpp
+++ b/Engine/Plugins/Developer/PerforceSourceControl/Source/PerforceSourceControl/Private/PerforceSourceControlProvider.cpp
@@ -38,31 +38,41 @@ namespace
 		IdleConnectionDisconnectSeconds,
 		TEXT("The number of seconds a perforce connection will be kept open without activity before being automatically disconnected"),
 		ECVF_Default);
+
+	void ParseCmdLineSetting(FSourceControlInitSettings& Settings)
+	{
+		if (!Settings.ShouldReadFromCmdLine())
+		{
+			return;
+		}
+
+		auto ParseCmdLineSetting = [&Settings](const TCHAR* SettingKey) -> void
+			{
+				FString SettingValue;
+				if (FParse::Value(FCommandLine::Get(), *WriteToString<64>(SettingKey, TEXT("=")), SettingValue))
+				{
+					Settings.AddSetting(SettingKey, SettingValue);
+				}
+			};
+
+		ParseCmdLineSetting(TEXT("P4Port"));
+		ParseCmdLineSetting(TEXT("P4User"));
+		ParseCmdLineSetting(TEXT("P4Client"));
+		ParseCmdLineSetting(TEXT("P4Host"));
+		ParseCmdLineSetting(TEXT("P4Passwd"));
+		ParseCmdLineSetting(TEXT("P4Changelist"));
+	}
 }
 
 FPerforceSourceControlProvider::FPerforceSourceControlProvider()
 	: PerforceSCCSettings(*this, FStringView())
-	, InitialSettings(FSourceControlInitSettings::EBehavior::OverrideExisting)
+	, InitialSettings(FSourceControlInitSettings::EBehavior::OverrideExisting, FSourceControlInitSettings::ECmdLineFlags::ReadAll)
 	, OwnerName(TEXT("Default"))
 	, bServerAvailable(false)
 	, bLoginError(false)
 	, PersistentConnection(nullptr)
 {
-	auto ParseCmdLineSetting = [this](const TCHAR* SettingKey) -> void
-	{
-		FString SettingValue;
-		if (FParse::Value(FCommandLine::Get(), *WriteToString<64>(SettingKey, TEXT("=")), SettingValue))
-		{
-			InitialSettings.AddSetting(SettingKey, SettingValue);
-		}
-	};
-
-	ParseCmdLineSetting(TEXT("P4Port"));
-	ParseCmdLineSetting(TEXT("P4User"));
-	ParseCmdLineSetting(TEXT("P4Client"));
-	ParseCmdLineSetting(TEXT("P4Host"));
-	ParseCmdLineSetting(TEXT("P4Passwd"));
-	ParseCmdLineSetting(TEXT("P4Changelist"));
+	ParseCmdLineSetting(InitialSettings);
 
 	AccessSettings().LoadSettings();
 }
@@ -75,6 +85,9 @@ FPerforceSourceControlProvider::FPerforceSourceControlProvider(const FStringView
 	, bLoginError(false)
 	, PersistentConnection(nullptr)
 {
+	
+	ParseCmdLineSetting(InitialSettings);
+
 	AccessSettings().SetAllowSave(InInitialSettings.CanWriteToConfigFile());
 	AccessSettings().SetAllowLoad(InInitialSettings.CanReadFromConfigFile());
 
diff --git a/Engine/Source/Developer/SourceControl/Private/SourceControlInitSettings.cpp b/Engine/Source/Developer/SourceControl/Private/SourceControlInitSettings.cpp
--- a/Engine/Source/Developer/SourceControl/Private/SourceControlInitSettings.cpp
+++ b/Engine/Source/Developer/SourceControl/Private/SourceControlInitSettings.cpp
@@ -10,6 +10,13 @@ FSourceControlInitSettings::FSourceControlInitSettings(EBehavior InBehavior)
 
 }
 
+FSourceControlInitSettings::FSourceControlInitSettings(EBehavior InBehavior, ECmdLineFlags InCmdLineFlags)
+	: OverrideBehavior(InBehavior)
+	, CmdLineFlags(InCmdLineFlags)
+{
+
+}
+
 void FSourceControlInitSettings::SetConfigBehavior(EConfigBehavior InBehavior)
 {
 	ConfigBehavior = InBehavior;
@@ -54,4 +61,14 @@ bool FSourceControlInitSettings::IsOverridden(FStringView SettingName) const
 {
 	const int32 Hash = GetTypeHash(SettingName);
 	return Settings.FindByHash(Hash, SettingName) != nullptr;
-}
\ No newline at end of file
+}
+
+void FSourceControlInitSettings::SetCmdLineFlags(ECmdLineFlags Flags)
+{
+	CmdLineFlags = Flags;
+}
+
+bool FSourceControlInitSettings::ShouldReadFromCmdLine() const
+{
+	return CmdLineFlags == ECmdLineFlags::ReadAll;
+}
diff --git a/Engine/Source/Developer/SourceControl/Public/SourceControlInitSettings.h b/Engine/Source/Developer/SourceControl/Public/SourceControlInitSettings.h
--- a/Engine/Source/Developer/SourceControl/Public/SourceControlInitSettings.h
+++ b/Engine/Source/Developer/SourceControl/Public/SourceControlInitSettings.h
@@ -19,7 +19,15 @@ public:
 		OverrideExisting,
 	};
 
-	enum class EConfigBehavior
+	enum class ECmdLineFlags : uint8
+	{
+		/** Do not read any settings from the commandline */
+		None,
+		/** Read all available settings from the commandline */
+		ReadAll
+	};
+
+	enum class EConfigBehavior : uint8
 	{
 		/** Can both read from, and save to the ini file*/
 		ReadWrite,
@@ -30,6 +38,7 @@ public:
 	};
 
 	FSourceControlInitSettings(EBehavior Behavior);
+	FSourceControlInitSettings(EBehavior Behavior, ECmdLineFlags CmdLineFlags);
 	~FSourceControlInitSettings() = default;
 
 	void SetConfigBehavior(EConfigBehavior Behavior);
@@ -43,10 +52,14 @@ public:
 	bool HasOverrides() const;
 	bool IsOverridden(FStringView SettingName) const;
 
+	void SetCmdLineFlags(ECmdLineFlags Flags);
+	bool ShouldReadFromCmdLine() const;
+
 private:
 
-	EBehavior OverrideBehavior = EBehavior::OverrideAll;
-	EConfigBehavior ConfigBehavior = EConfigBehavior::ReadWrite;
+	EBehavior OverrideBehavior		= EBehavior::OverrideAll;
+	ECmdLineFlags CmdLineFlags		= ECmdLineFlags::None;
+	EConfigBehavior ConfigBehavior	= EConfigBehavior::ReadWrite;
 
 	TMap<FString, FString> Settings;
 };
