Add project directory to p4 workspace error message

diff --git a/Engine/Plugins/Developer/PerforceSourceControl/Source/PerforceSourceControl/Private/PerforceSourceControlOperations.cpp b/Engine/Plugins/Developer/PerforceSourceControl/Source/PerforceSourceControl/Private/PerforceSourceControlOperations.cpp
--- a/Engine/Plugins/Developer/PerforceSourceControl/Source/PerforceSourceControl/Private/PerforceSourceControlOperations.cpp
+++ b/Engine/Plugins/Developer/PerforceSourceControl/Source/PerforceSourceControl/Private/PerforceSourceControlOperations.cpp
@@ -391,7 +391,7 @@ static bool CheckWorkspaceRecordSet(const FP4RecordSet& InRecords, TArray<FText>
 		else
 		{
 			const FString& Client = Record(TEXT("Client"));
-			OutNotificationText = FText::Format(LOCTEXT("WorkspaceError", "Workspace '{0}' does not map into this project's directory."), FText::FromString(Client));
+			OutNotificationText = FText::Format(LOCTEXT("WorkspaceError", "Workspace '{0}' does not map into this project's directory '{1}'"), FText::FromString(Client), FText::FromString(ApplicationPath));
 			OutErrorMessages.Add(OutNotificationText);
 			OutErrorMessages.Add(LOCTEXT("WorkspaceHelp", "You should set your workspace up to map to a directory at or above the project's directory."));
 		}
