Changed p4 change -i to p4 change and reusing the template that the server provides

This ensures that we use the changelist type from the server which was being dropped

#rnx

diff --git a/Engine/Plugins/Developer/PerforceSourceControl/Source/PerforceSourceControl/Private/PerforceConnection.cpp b/Engine/Plugins/Developer/PerforceSourceControl/Source/PerforceSourceControl/Private/PerforceConnection.cpp
--- a/Engine/Plugins/Developer/PerforceSourceControl/Source/PerforceSourceControl/Private/PerforceConnection.cpp
+++ b/Engine/Plugins/Developer/PerforceSourceControl/Source/PerforceSourceControl/Private/PerforceConnection.cpp
@@ -348,7 +348,60 @@ public:
 		,	ChangelistNumber(0)
 		,	Files(InFiles)
 		,	P4Client(InP4Client)
+	{}
+
+	/** Called by P4API on "change" command. FileSys is a file pointer to the input contents that p4 change will use. This comes with a default template for the server (see p4 change)*/
+	virtual void Edit(FileSys* f1, Error* e) override
 	{
+		StrBuf ReadBuffer;
+		f1->ReadFile(&ReadBuffer, e);
+
+		FString TemplateContents = TO_TCHAR(ReadBuffer.Value(), IsUnicodeServer());
+		TArray<FString> Lines;
+		TemplateContents.ParseIntoArrayLines(Lines, false);
+
+		TStringBuilder<2048> FinalContents;
+
+		for (size_t i = 0; i < Lines.Num(); ++i)
+		{
+			const FString& Line = Lines[i];
+
+			// Ignore comments in the template and keep everything else
+			if (!Line.StartsWith(TEXT("#")))
+			{
+				FinalContents << Line << TEXT("\n");
+
+				// When we find the Description, replace it with our own.
+				if (Line.StartsWith(TEXT("Description:")))
+				{
+					TArray<FString> DescLines;
+					Description.ToString().ParseIntoArray(DescLines, TEXT("\n"), false);
+					for (const FString& DescLine : DescLines)
+					{
+						FinalContents << TEXT("\t") << DescLine << TEXT("\n");
+					}
+
+					// Skip the next line after "Description:" which will be "<enter description here>"
+					++i;
+				}
+			}
+		}
+
+		// Add files if there are any
+		if (Files.Num() != 0)
+		{
+			FinalContents << TEXT("Files:\n");
+
+			for (const FString& FileName : Files)
+			{
+				FinalContents << TEXT("\t") << FileName << TEXT("\n");
+			}
+
+			FinalContents << TEXT("\n");
+		}
+
+		StrBuf WriteBuffer(FROM_TCHAR(*FinalContents, IsUnicodeServer()));
+		f1->WriteFile(&WriteBuffer, e);
 	}
 
 	/** Called by P4API when the changelist is created. */
@@ -372,43 +425,6 @@ public:
 		FP4ClientUser::Message(err);
 	}
 
-
-	/** Called by P4API on "change -i" command. OutBuffer is filled with changelist specification text. */
-	virtual void InputData(StrBuf* OutBuffer, Error* OutError) override
-	{
-		FString OutputDesc;
-		OutputDesc += TEXT("Change:\tnew\n\n");
-		OutputDesc += TEXT("Client:\t");
-		OutputDesc += TO_TCHAR(P4Client.GetClient().Text(), IsUnicodeServer());
-		OutputDesc += TEXT("\n\n");
-		OutputDesc += TEXT("User:\t");
-		OutputDesc += TO_TCHAR(P4Client.GetUser().Text(), IsUnicodeServer());
-		OutputDesc += TEXT("\n\n");
-		OutputDesc += TEXT("Status:\tnew\n\n");
-		OutputDesc += TEXT("Description:\n");
-		{
-			TArray<FString> DescLines;
-			Description.ToString().ParseIntoArray(DescLines, TEXT("\n"), false);
-			for (const FString& DescLine : DescLines)
-			{
-				OutputDesc += TEXT("\t");
-				OutputDesc += DescLine;
-				OutputDesc += TEXT("\n");
-			}
-		}
-		OutputDesc += TEXT("\n");
-		OutputDesc += TEXT("Files:\n");
-		for (const FString& FileName : Files)
-		{
-			OutputDesc += TEXT("\t");
-			OutputDesc += FileName;
-			OutputDesc += TEXT("\n");
-		}
-		OutputDesc += TEXT("\n");
-
-		OutBuffer->Append(FROM_TCHAR(*OutputDesc, IsUnicodeServer()));
-	}
-
 	FText Description;
 	int32 ChangelistNumber;
 	TArray<FString> Files;
@@ -1203,9 +1219,6 @@ int32 FPerforceConnection::CreatePendingChangelist(const FText &Description, con
 	FP4RecordSet Records;
 	EP4ClientUserFlags Flags = bIsUnicode ? EP4ClientUserFlags::UnicodeServer : EP4ClientUserFlags::None;
 
-	const char *ArgV[] = { "-i" };
-	P4Client.SetArgv(1, const_cast<char*const*>(ArgV));
-
 	FP4KeepAlive KeepAlive(InIsCancelled);
 	P4Client.SetBreak(&KeepAlive);
 
