Fix Perforce create changelist moving all files from the default CL to the new changelist instead of creating an empty changelist or moving only the selected files

diff --git a/Engine/Plugins/Developer/PerforceSourceControl/Source/PerforceSourceControl/Private/PerforceConnection.cpp b/Engine/Plugins/Developer/PerforceSourceControl/Source/PerforceSourceControl/Private/PerforceConnection.cpp
--- a/Engine/Plugins/Developer/PerforceSourceControl/Source/PerforceSourceControl/Private/PerforceConnection.cpp
+++ b/Engine/Plugins/Developer/PerforceSourceControl/Source/PerforceSourceControl/Private/PerforceConnection.cpp
@@ -361,6 +361,7 @@ public:
 		TemplateContents.ParseIntoArrayLines(Lines, false);
 
 		TStringBuilder<2048> FinalContents;
+		bool bFileSectionPresent= false;
 
 		for (size_t i = 0; i < Lines.Num(); ++i)
 		{
@@ -369,11 +370,11 @@ public:
 			// Ignore comments in the template and keep everything else
 			if (!Line.StartsWith(TEXT("#")))
 			{
-				FinalContents << Line << TEXT("\n");
-
 				// When we find the Description, replace it with our own.
 				if (Line.StartsWith(TEXT("Description:")))
 				{
+					FinalContents << Line << TEXT("\n");
+
 					TArray<FString> DescLines;
 					Description.ToString().ParseIntoArray(DescLines, TEXT("\n"), false);
 					for (const FString& DescLine : DescLines)
@@ -384,11 +385,40 @@ public:
 					// Skip the next line after "Description:" which will be "<enter description here>"
 					++i;
 				}
+				// When we find the File section, remove it completely to create an empty CL
+				// Or replace them with our own files if we were provided any
+				else if (Line.StartsWith(TEXT("Files:")))
+				{
+					bFileSectionPresent = true;
+
+					if (Files.Num() != 0)
+					{
+						FinalContents << Line << TEXT("\n");
+
+						for (const FString& FileName : Files)
+						{
+							FinalContents << TEXT("\t") << FileName << TEXT("\n");
+						}
+
+						FinalContents << TEXT("\n");
+					}
+
+					// Skip the default files up to empty line
+					// Sections in the p4 change command are separated by empty lines
+					do
+					{
+						++i;
+					}
+					while (i < Lines.Num() && !Lines[i].IsEmpty());
+				}
+				else
+				{
+					FinalContents << Line << TEXT("\n");
+				}
 			}
 		}
 
-		// Add files if there are any
-		if (Files.Num() != 0)
+		if (!bFileSectionPresent && Files.Num() != 0)
 		{
 			FinalContents << TEXT("Files:\n");
 
