Hotfixing a warning in a previous Submit

#rnx

diff --git a/Engine/Plugins/Developer/PerforceSourceControl/Source/PerforceSourceControl/Private/PerforceSourceControlOperations.cpp b/Engine/Plugins/Developer/PerforceSourceControl/Source/PerforceSourceControl/Private/PerforceSourceControlOperations.cpp
--- a/Engine/Plugins/Developer/PerforceSourceControl/Source/PerforceSourceControl/Private/PerforceSourceControlOperations.cpp
+++ b/Engine/Plugins/Developer/PerforceSourceControl/Source/PerforceSourceControl/Private/PerforceSourceControlOperations.cpp
@@ -277,7 +277,7 @@ static void ParseRecordSetForState(const FP4RecordSet& InRecords, TMap<FString,
 		const FString& DepotFileName = ClientRecord(TEXT("depotFile"));
 		const FString& Action = ClientRecord(TEXT("action"));
 
-		check(ClientFileName.Len() || bConsiderDepotFile ? DepotFileName.Len() : false);
+		check(ClientFileName.Len() || (bConsiderDepotFile ? DepotFileName.Len() : false));
 		FString FullPath = ClientFileName.Len() ? ClientFileName : (bConsiderDepotFile ? DepotFileName : FString());
 		FPaths::NormalizeFilename(FullPath);
 
