Hotfix crash/check in source control when project depends on EditorDataStorage but not EditorDataStorageFeatures

diff --git a/Engine/Plugins/Developer/PerforceSourceControl/Source/PerforceSourceControl/Private/PerforceSourceControlOperations.cpp b/Engine/Plugins/Developer/PerforceSourceControl/Source/PerforceSourceControl/Private/PerforceSourceControlOperations.cpp
--- a/Engine/Plugins/Developer/PerforceSourceControl/Source/PerforceSourceControl/Private/PerforceSourceControlOperations.cpp
+++ b/Engine/Plugins/Developer/PerforceSourceControl/Source/PerforceSourceControl/Private/PerforceSourceControlOperations.cpp
@@ -2039,7 +2039,14 @@ bool FPerforceUpdateStatusWorker::UpdateStates() const
 			return;
 		}
 
-		auto GetRevisionControlRow = [DataStorage](const FString& InFilename) -> RowHandle
+		const RowHandle RevisionControlTable = DataStorage->FindTable(FName("Editor_RevisionControlTable"));
+		if (RevisionControlTable == InvalidTableHandle)
+		{
+			return;
+		}
+		
+
+		auto GetRevisionControlRow = [DataStorage, RevisionControlTable](const FString& InFilename) -> RowHandle
 		{
 			FString Filename = FPaths::SetExtension(InFilename, "");
 			FPaths::NormalizeFilename(Filename);
@@ -2050,8 +2057,7 @@ bool FPerforceUpdateStatusWorker::UpdateStates() const
 
 			if (!DataStorage->IsRowAvailable(Row))
 			{
-				static RowHandle Table = DataStorage->FindTable(FName("Editor_RevisionControlTable"));
-				Row = DataStorage->AddRow(Table);
+				Row = DataStorage->AddRow(RevisionControlTable);
 				DataStorage->IndexRow(Index, Row);
 			}
 			return Row;
