Added a SetQuiet boolean to FGetFileList Source Control Operation to be able to prevent messages "... - no such file(s)." from generating errors in the logs.

diff --git a/Engine/Plugins/Developer/PerforceSourceControl/Source/PerforceSourceControl/Private/PerforceSourceControlOperations.cpp b/Engine/Plugins/Developer/PerforceSourceControl/Source/PerforceSourceControl/Private/PerforceSourceControlOperations.cpp
--- a/Engine/Plugins/Developer/PerforceSourceControl/Source/PerforceSourceControl/Private/PerforceSourceControlOperations.cpp
+++ b/Engine/Plugins/Developer/PerforceSourceControl/Source/PerforceSourceControl/Private/PerforceSourceControlOperations.cpp
@@ -812,6 +812,10 @@ bool FPerforceGetFileListWorker::Execute(FPerforceSourceControlCommand& InComman
 		}
 		FP4RecordSet Records;
 		InCommand.bCommandSuccessful = Connection.RunCommand(TEXT("files"), Parameters, Records, InCommand.ResultInfo, FOnIsCancelled::CreateRaw(&InCommand, &FPerforceSourceControlCommand::IsCanceled), InCommand.bConnectionDropped);
+		if (Operation->ShouldBeQuiet())
+		{
+			RemoveRedundantErrors(InCommand, TEXT(" - no such file(s)."), false);
+		}
 		bool bConsiderDepotFiles = Operation->GetMethodUsed() == FGetFileList::FileRegexSearch;
 		ParseRecordSetForState(Records, OutResults, bConsiderDepotFiles);
 
diff --git a/Engine/Source/Developer/SourceControl/Private/SourceControlHelpers.cpp b/Engine/Source/Developer/SourceControl/Private/SourceControlHelpers.cpp
--- a/Engine/Source/Developer/SourceControl/Private/SourceControlHelpers.cpp
+++ b/Engine/Source/Developer/SourceControl/Private/SourceControlHelpers.cpp
@@ -1725,6 +1725,7 @@ bool USourceControlHelpers::GetFilesInDepotAtPaths(const TArray<FString>& Paths,
 
 			TSharedRef<FGetFileList, ESPMode::ThreadSafe> Operation = ISourceControlOperation::Create<FGetFileList>();
 			Operation->SetIncludeDeleted(bIncludeDeleted);
+			Operation->SetQuiet(bSilent);
 			Operation->SetSearchPattern(PathsToQuery);
 			if (bIsFileRegexSearch)
 			{
diff --git a/Engine/Source/Developer/SourceControl/Public/SourceControlOperations.h b/Engine/Source/Developer/SourceControl/Public/SourceControlOperations.h
--- a/Engine/Source/Developer/SourceControl/Public/SourceControlOperations.h
+++ b/Engine/Source/Developer/SourceControl/Public/SourceControlOperations.h
@@ -176,6 +176,16 @@ public:
 		return bIncludeDeleted;
 	}
 
+	void SetQuiet(bool bInQuiet)
+	{
+		bForceQuiet = bInQuiet;
+	}
+
+	bool ShouldBeQuiet() const
+	{
+		return bForceQuiet;
+	}
+
 	void SetMethodUsed(const EGetFileListMethod& InMethodUsed)
 	{
 		MethodUsed = InMethodUsed;
@@ -213,6 +223,9 @@ protected:
 	/** Include deleted files in the list. */
 	bool bIncludeDeleted = false;
 
+	/** Force disable any logs generated by that operation. */
+	bool bForceQuiet = false;
+
 	/** Stored result of the operation */
 	TArray<FString> FilesList;
 
