Remove unused parameter.

#rnx

diff --git a/Engine/Plugins/Developer/PerforceSourceControl/Source/PerforceSourceControl/Private/PerforceSourceControlLabel.cpp b/Engine/Plugins/Developer/PerforceSourceControl/Source/PerforceSourceControl/Private/PerforceSourceControlLabel.cpp
--- a/Engine/Plugins/Developer/PerforceSourceControl/Source/PerforceSourceControl/Private/PerforceSourceControlLabel.cpp
+++ b/Engine/Plugins/Developer/PerforceSourceControl/Source/PerforceSourceControl/Private/PerforceSourceControlLabel.cpp
@@ -16,7 +16,7 @@ const FString& FPerforceSourceControlLabel::GetName() const
 	return Name;
 }	
 
-static void ParseFilesResults(FPerforceSourceControlProvider& SCCProvider, const FP4RecordSet& InRecords, const FString& InClientRoot, TArray< TSharedRef<ISourceControlRevision, ESPMode::ThreadSafe> >& OutRevisions)
+static void ParseFilesResults(FPerforceSourceControlProvider& SCCProvider, const FP4RecordSet& InRecords, TArray< TSharedRef<ISourceControlRevision, ESPMode::ThreadSafe> >& OutRevisions)
 {
 	// Iterate over each record found as a result of the command, parsing it for relevant information
 	for (const FP4Record& ClientRecord : InRecords)
@@ -66,7 +66,7 @@ bool FPerforceSourceControlLabel::GetFileRevisions( const TArray<FString>& InFil
 		bCommandOK = Connection.RunCommand(TEXT("files"), Parameters, Records, ResultInfo, FOnIsCancelled(), bConnectionDropped);
 		if(bCommandOK)
 		{
-			ParseFilesResults(GetSCCProvider(), Records, Connection.ClientRoot, OutRevisions);
+			ParseFilesResults(GetSCCProvider(), Records, OutRevisions);
 		}
 		else
 		{
