Fix perforce source control module from adding double slashes to local file paths if the workspace root is set to a volume root.

### Problem
- Previously when storing the workspace root we assumed it would be a directory path and so normalized it to make sure that there was no trailing slash so that our file path parsing code in ::ParseOpenedResults would be consistent. Note that this was changed back in CL 23945151, before that the root could have a trailing slash or not which could cause issues.
- This assumption didn't work if the user decided to use a volume root (i.e D:\) as the workspace root, as in our path utilities normalizing code assumes that volume roots should end with a trailing slash. This in turn would result in local file paths being generated by ::ParseOpenedResults that contained duplicate slashes in the path.
- It is also worth noting that other source control implementations were specifically adding trailing slashes to their equivalents of the workspace root, so forcing perforce roots to not have the trailing slash was inconsistent with other implementations too.

### Solution
- We now enforce that the workspace root ends with a trailing slash unless the workspace root is null, and parse the results accordingly.
- We also removed duplicate slashes from the client root as perforce seems to technically support that even if p4v has issues with workspaces set up like this. At the very least this will move any issues from that to be pure perforce issues and not in our own parsing code.

diff --git a/Engine/Plugins/Developer/PerforceSourceControl/Source/PerforceSourceControl/Private/PerforceConnection.cpp b/Engine/Plugins/Developer/PerforceSourceControl/Source/PerforceSourceControl/Private/PerforceConnection.cpp
--- a/Engine/Plugins/Developer/PerforceSourceControl/Source/PerforceSourceControl/Private/PerforceConnection.cpp
+++ b/Engine/Plugins/Developer/PerforceSourceControl/Source/PerforceSourceControl/Private/PerforceConnection.cpp
@@ -1420,7 +1420,17 @@ void FPerforceConnection::EstablishConnection(const FPerforceConnectionInfo& InC
 				UE_LOG(LogSourceControl, Verbose, TEXT(" ... getting clientroot" ));
 				ClientRoot = Records[0](TEXT("clientRoot"));
 
-				FPaths::NormalizeDirectoryName(ClientRoot);
+				// Make sure that workspace roots end with a seperator or are null
+				if (ClientRoot != TEXT("null"))
+				{
+					FPaths::NormalizeFilename(ClientRoot);		// Don't use NormalizeDirectoryName, we want to preserve trailing slashes
+					FPaths::RemoveDuplicateSlashes(ClientRoot); // Perforce client specs support duplicate slashes
+
+					if (!ClientRoot.EndsWith(TEXT("/")))
+					{
+						ClientRoot.Append(TEXT("/"));
+					}
+				}
 			}
 		}
 	}
diff --git a/Engine/Plugins/Developer/PerforceSourceControl/Source/PerforceSourceControl/Private/PerforceSourceControlOperations.cpp b/Engine/Plugins/Developer/PerforceSourceControl/Source/PerforceSourceControl/Private/PerforceSourceControlOperations.cpp
--- a/Engine/Plugins/Developer/PerforceSourceControl/Source/PerforceSourceControl/Private/PerforceSourceControlOperations.cpp
+++ b/Engine/Plugins/Developer/PerforceSourceControl/Source/PerforceSourceControl/Private/PerforceSourceControlOperations.cpp
@@ -1548,7 +1548,7 @@ static void ParseOpenedResults(const FP4RecordSet& InRecords, const FString& Cli
 
 		// Convert the depot file name to a local file name
 		FString FullPath = ClientFileName;
-		const FString PathRoot = FString::Printf(TEXT("//%s"), *ClientName);
+		const FString PathRoot = FString::Printf(TEXT("//%s/"), *ClientName);
 
 		if (FullPath.StartsWith(PathRoot))
 		{
@@ -1557,7 +1557,7 @@ static void ParseOpenedResults(const FP4RecordSet& InRecords, const FString& Cli
 			{
 				// Null clients use the pattern in PathRoot: //Workspace/FileName
 				// Here we chop off the '//Workspace/' to return the workspace filename
-				FullPath.RightChopInline(PathRoot.Len() + 1, EAllowShrinking::No);
+				FullPath.RightChopInline(PathRoot.Len(), EAllowShrinking::No);
 			}
 			else
 			{
