Perforce will now return the correct workspace root after a user has changed the workspace via the UX.

- Calls to FPerforceSourceControlProvider::GetStatus on the GameThread will end up using the persistent connection to determine the current workspace root. Before this change the persistent connection was not getting updated when the user changed the workspace via the UX, so we'd return the root of the original workspace when the editor started up, if one had been set.
- Added a call to reset the persistent connection when the FConnect command is run, we can then let the persistent connection be recreated when it is next requested.
- Changed FPerforceSourceControlProvider::PersistentConnection from a raw pointer to TUniquePtr to make management of it a bit easier.

diff --git a/Engine/Plugins/Developer/PerforceSourceControl/Source/PerforceSourceControl/Private/PerforceSourceControlOperations.cpp b/Engine/Plugins/Developer/PerforceSourceControl/Source/PerforceSourceControl/Private/PerforceSourceControlOperations.cpp
--- a/Engine/Plugins/Developer/PerforceSourceControl/Source/PerforceSourceControl/Private/PerforceSourceControlOperations.cpp
+++ b/Engine/Plugins/Developer/PerforceSourceControl/Source/PerforceSourceControl/Private/PerforceSourceControlOperations.cpp
@@ -426,6 +426,8 @@ bool FPerforceConnectWorker::Execute(FPerforceSourceControlCommand& InCommand)
 {
 	TRACE_CPUPROFILER_EVENT_SCOPE(FPerforceConnectWorker::Execute);
 
+	GetSCCProvider().ResetPersistentConnection();
+
 	if (InCommand.ConnectionInfo.Workspace.IsEmpty())
 	{
 		// Running FConnect for perforce does not actually connect to the server but
diff --git a/Engine/Plugins/Developer/PerforceSourceControl/Source/PerforceSourceControl/Private/PerforceSourceControlProvider.cpp b/Engine/Plugins/Developer/PerforceSourceControl/Source/PerforceSourceControl/Private/PerforceSourceControlProvider.cpp
--- a/Engine/Plugins/Developer/PerforceSourceControl/Source/PerforceSourceControl/Private/PerforceSourceControlProvider.cpp
+++ b/Engine/Plugins/Developer/PerforceSourceControl/Source/PerforceSourceControl/Private/PerforceSourceControlProvider.cpp
@@ -114,8 +114,7 @@ void FPerforceSourceControlProvider::Close()
 	if ( PersistentConnection )
 	{
 		PersistentConnection->Disconnect();
-		delete PersistentConnection;
-		PersistentConnection = NULL;
+		PersistentConnection.Reset();
 	}
 
 	// clear the cache
@@ -222,16 +221,16 @@ bool FPerforceSourceControlProvider::EstablishPersistentConnection()
 	FPerforceConnectionInfo ConnectionInfo = AccessSettings().GetConnectionInfo();
 
 	bool bIsValidConnection = false;
-	if ( !PersistentConnection )
+	if (!PersistentConnection)
 	{
-		PersistentConnection = new FPerforceConnection(ConnectionInfo, *this);
+		PersistentConnection = MakeUnique<FPerforceConnection>(ConnectionInfo, *this);
 	}
 
 	bIsValidConnection = PersistentConnection->IsValidConnection();
-	if ( !bIsValidConnection )
+	if (!bIsValidConnection)
 	{
-		delete PersistentConnection;
-		PersistentConnection = new FPerforceConnection(ConnectionInfo, *this);
+		PersistentConnection.Reset();
+		PersistentConnection = MakeUnique<FPerforceConnection>(ConnectionInfo, *this);
 		bIsValidConnection = PersistentConnection->IsValidConnection();
 	}
 
@@ -239,6 +238,11 @@ bool FPerforceSourceControlProvider::EstablishPersistentConnection()
 	return bIsValidConnection;
 }
 
+void FPerforceSourceControlProvider::ResetPersistentConnection()
+{
+	PersistentConnection.Reset();
+}
+
 ISourceControlProvider::FInitResult FPerforceSourceControlProvider::ParseCommandLineSettings(EInitFlags InitFlags)
 {
 	ISourceControlProvider::FInitResult Result;
diff --git a/Engine/Plugins/Developer/PerforceSourceControl/Source/PerforceSourceControl/Private/PerforceSourceControlProvider.h b/Engine/Plugins/Developer/PerforceSourceControl/Source/PerforceSourceControl/Private/PerforceSourceControlProvider.h
--- a/Engine/Plugins/Developer/PerforceSourceControl/Source/PerforceSourceControl/Private/PerforceSourceControlProvider.h
+++ b/Engine/Plugins/Developer/PerforceSourceControl/Source/PerforceSourceControl/Private/PerforceSourceControlProvider.h
@@ -12,6 +12,7 @@
 #include "PerforceSourceControlState.h"
 #include "SourceControlInitSettings.h"
 
+class FPerforceConnection;
 class FPerforceSourceControlCommand;
 
 class FPerforceSourceControlProvider : public ISourceControlProvider
@@ -96,10 +97,12 @@ public:
 	 */
 	bool EstablishPersistentConnection();
 
+	void ResetPersistentConnection();
+
 	/** Get the persistent connection, if any */
 	class FPerforceConnection* GetPersistentConnection()
 	{
-		return PersistentConnection;
+		return PersistentConnection.Get();
 	}
 
 	/** Remove a named file from the state cache */
@@ -129,7 +132,7 @@ private:
 	/** 
 	 * Logs any messages that a command needs to output.
 	 */
-	void OutputCommandMessages(const class FPerforceSourceControlCommand& InCommand) const;
+	void OutputCommandMessages(const FPerforceSourceControlCommand& InCommand) const;
 
 	/**
 	 * Loads user/SCC information from the INI file and can attempt to make a connection to the server if requested by
@@ -142,12 +145,12 @@ private:
 	 * This really doesn't execute synchronously; rather it adds the command to the queue & does not return until
 	 * the command is completed.
 	 */
-	ECommandResult::Type ExecuteSynchronousCommand(class FPerforceSourceControlCommand& InCommand, const FText& Task, bool bSuppressResponseMsg);
+	ECommandResult::Type ExecuteSynchronousCommand(FPerforceSourceControlCommand& InCommand, const FText& Task, bool bSuppressResponseMsg);
 
 	/**
 	 * Run a command synchronously or asynchronously.
 	 */
-	ECommandResult::Type IssueCommand(class FPerforceSourceControlCommand& InCommand, const bool bSynchronous);
+	ECommandResult::Type IssueCommand(FPerforceSourceControlCommand& InCommand, const bool bSynchronous);
 
 private:
 
@@ -179,14 +182,14 @@ private:
 	TArray<FText> LastErrors;
 
 	/** A pointer to the persistent P4 connection for synchronous operations */
-	class FPerforceConnection* PersistentConnection;
+	TUniquePtr<FPerforceConnection> PersistentConnection;
 
 	/** State cache */
 	TMap<FString, TSharedRef<class FPerforceSourceControlState, ESPMode::ThreadSafe> > StateCache;
 	TMap<FPerforceSourceControlChangelist, TSharedRef<class FPerforceSourceControlChangelistState, ESPMode::ThreadSafe> > ChangelistsStateCache;
 
 	/** Queue for commands given by the main thread */
-	TArray < FPerforceSourceControlCommand* > CommandQueue;
+	TArray <FPerforceSourceControlCommand* > CommandQueue;
 
 	/** For notifying when the source control states in the cache have changed */
 	FSourceControlStateChanged OnSourceControlStateChanged;
