Add basic support for ISourceControlOperation::ShouldLogErrors in the perforce source control module .

#rnx

- FPerforceSourceControlProvider::OutputCommandMessages now skips logging errors if ISourceControlOperation::ShouldLogErrors returns false. Info logging is not affected.
- When creating a perforce connection for an operation we check if ISourceControlOperation::ShouldLogErrors returns false. If it does then we pass in the command's FSourceControlResultInfo to the connection creation. Then when establishing the connection if a valid FSourceControlResultInfo pointer has been provided we will store any errors there. If the FSourceControlResultInfo pointer is a nullptr then the errors will be immediately logged as before.
- This submit does not do an audit of each operation to check for errant error logging, that will be tackled as a future work item.

diff --git a/Engine/Plugins/Developer/PerforceSourceControl/Source/PerforceSourceControl/Private/PerforceConnection.cpp b/Engine/Plugins/Developer/PerforceSourceControl/Source/PerforceSourceControl/Private/PerforceConnection.cpp
--- a/Engine/Plugins/Developer/PerforceSourceControl/Source/PerforceSourceControl/Private/PerforceConnection.cpp
+++ b/Engine/Plugins/Developer/PerforceSourceControl/Source/PerforceSourceControl/Private/PerforceConnection.cpp
@@ -767,13 +767,13 @@ static void FinalizeErrors(const FPerforceSourceControlProvider& SCCProvider, co
 	OutConnectionErrors.AdditionalErrors.Add(FText::Join(LOCTEXT("P4_Delim", ", "), Components));
 }
 
-FPerforceConnection::FPerforceConnection(const FPerforceConnectionInfo& InConnectionInfo, FPerforceSourceControlProvider& InSCCProvider)
+FPerforceConnection::FPerforceConnection(const FPerforceConnectionInfo& InConnectionInfo, FPerforceSourceControlProvider& InSCCProvider, FSourceControlResultInfo* OutResultInfo)
 	: bEstablishedConnection(false)
 	, bIsUnicode(false)
 	, LatestCommunicateTime(0.0)
 	, SCCProvider(InSCCProvider)
 {
-	EstablishConnection(InConnectionInfo);
+	EstablishConnection(InConnectionInfo, OutResultInfo);
 }
 
 FPerforceConnection::~FPerforceConnection()
@@ -787,7 +787,7 @@ bool FPerforceConnection::AutoDetectWorkspace(const FPerforceConnectionInfo& InC
 
 	OutErrorMessages.Reset();
 
-	FPerforceConnection Connection(InConnectionInfo, SCCProvider);
+	FPerforceConnection Connection(InConnectionInfo, SCCProvider, nullptr);
 	TArray<FString> ClientSpecList;
 	FSourceControlResultInfo ResultInfo;
 
@@ -1327,7 +1327,7 @@ bool FPerforceConnection::CreateWorkspace(FStringView WorkspaceSpec, FOnIsCancel
 	return !OutResultInfo.HasErrors();
 }
 
-void FPerforceConnection::EstablishConnection(const FPerforceConnectionInfo& InConnectionInfo)
+void FPerforceConnection::EstablishConnection(const FPerforceConnectionInfo& InConnectionInfo, FSourceControlResultInfo* OutResultInfo)
 {
 	TRACE_CPUPROFILER_EVENT_SCOPE(FPerforceConnection::EstablishConnection);
 
@@ -1369,20 +1369,28 @@ void FPerforceConnection::EstablishConnection(const FPerforceConnectionInfo& InC
 		StrBuf ErrorMessage;
 		P4Error.Fmt(&ErrorMessage);
 
-		UE_LOG(LogSourceControl, Error, TEXT("P4ERROR: Invalid connection to server."));
-		UE_LOG(LogSourceControl, Error, TEXT("%s"), ANSI_TO_TCHAR(ErrorMessage.Text()));
+		if (OutResultInfo != nullptr)
+		{
+			OutResultInfo->ErrorMessages.Add(FText::FromString(TEXT("P4ERROR: Invalid connection to server.")));
+			OutResultInfo->ErrorMessages.Add(FText::FromString(ANSI_TO_TCHAR(ErrorMessage.Text())));
+		}
+		else
+		{
+			UE_LOG(LogSourceControl, Error, TEXT("P4ERROR: Invalid connection to server."));
+			UE_LOG(LogSourceControl, Error, TEXT("%s"), ANSI_TO_TCHAR(ErrorMessage.Text()));
+		}
 	}
 	else
 	{
 		TArray<FString> Params;
-		FSourceControlResultInfo ResultInfo;
+		FSourceControlResultInfo DummyResultInfo;
 		FP4RecordSet Records;
 		bool bConnectionDropped = false;
 		const ERunCommandFlags RunFlags = ERunCommandFlags::DisableCommandLogging;
 
 		UE_LOG(LogSourceControl, Verbose, TEXT(" ... checking unicode status" ));
 
-		if (RunCommand(TEXT("info"), Params, Records, nullptr, ResultInfo, FOnIsCancelled(), bConnectionDropped, RunFlags))
+		if (RunCommand(TEXT("info"), Params, Records, nullptr, DummyResultInfo, FOnIsCancelled(), bConnectionDropped, RunFlags))
 		{
 			// Get character encoding
 			bIsUnicode = Records[0].Find(TEXT("unicode")) != nullptr;
@@ -1415,7 +1423,7 @@ void FPerforceConnection::EstablishConnection(const FPerforceConnectionInfo& InC
 			// Gather the client root
 			UE_LOG(LogSourceControl, Verbose, TEXT(" ... getting info" ));
 			bConnectionDropped = false;
-			if (RunCommand(TEXT("info"), Params, Records, nullptr, ResultInfo, FOnIsCancelled(), bConnectionDropped, RunFlags))
+			if (RunCommand(TEXT("info"), Params, Records, nullptr, DummyResultInfo, FOnIsCancelled(), bConnectionDropped, RunFlags))
 			{
 				UE_LOG(LogSourceControl, Verbose, TEXT(" ... getting clientroot" ));
 				ClientRoot = Records[0](TEXT("clientRoot"));
@@ -1452,7 +1460,9 @@ FScopedPerforceConnection::FScopedPerforceConnection( FPerforceSourceControlComm
 	: Connection(nullptr)
 	, Concurrency(InCommand.Concurrency)
 {
-	Initialize(InCommand.ConnectionInfo, InCommand.Worker->GetSCCProvider());
+	FSourceControlResultInfo* ResultInfo = InCommand.Operation->ShouldLogErrors() ? nullptr : &InCommand.ResultInfo;
+
+	Initialize(InCommand.ConnectionInfo, InCommand.Worker->GetSCCProvider(), ResultInfo);
 	if (IsValid())
 	{
 		InCommand.MarkConnectionAsSuccessful();
@@ -1463,16 +1473,16 @@ FScopedPerforceConnection::FScopedPerforceConnection(EConcurrency::Type InConcur
 	: Connection(nullptr)
 	, Concurrency(InConcurrency)
 {
-	Initialize(SCCProvider.AccessSettings().GetConnectionInfo(), SCCProvider);
+	Initialize(SCCProvider.AccessSettings().GetConnectionInfo(), SCCProvider, nullptr);
 }
 
-void FScopedPerforceConnection::Initialize( const FPerforceConnectionInfo& InConnectionInfo, FPerforceSourceControlProvider& SCCProvider)
+void FScopedPerforceConnection::Initialize( const FPerforceConnectionInfo& InConnectionInfo, FPerforceSourceControlProvider& SCCProvider, FSourceControlResultInfo* OutResultInfo)
 {
 	if(Concurrency == EConcurrency::Synchronous)
 	{
 		// Synchronous commands reuse the same persistent connection to reduce
 		// the number of expensive connection attempts.
-		if (SCCProvider.EstablishPersistentConnection() )
+		if (SCCProvider.EstablishPersistentConnection(OutResultInfo) )
 		{
 			Connection = SCCProvider.GetPersistentConnection();
 		}
@@ -1481,7 +1491,7 @@ void FScopedPerforceConnection::Initialize( const FPerforceConnectionInfo& InCon
 	{
 		// Async commands form a new connection for each attempt because
 		// using the persistent connection is not threadsafe
-		FPerforceConnection* NewConnection = new FPerforceConnection(InConnectionInfo, SCCProvider);
+		FPerforceConnection* NewConnection = new FPerforceConnection(InConnectionInfo, SCCProvider, OutResultInfo);
 		if ( NewConnection->IsValidConnection() )
 		{
 			Connection = NewConnection;
diff --git a/Engine/Plugins/Developer/PerforceSourceControl/Source/PerforceSourceControl/Private/PerforceConnection.h b/Engine/Plugins/Developer/PerforceSourceControl/Source/PerforceSourceControl/Private/PerforceConnection.h
--- a/Engine/Plugins/Developer/PerforceSourceControl/Source/PerforceSourceControl/Private/PerforceConnection.h
+++ b/Engine/Plugins/Developer/PerforceSourceControl/Source/PerforceSourceControl/Private/PerforceConnection.h
@@ -67,7 +67,8 @@ class FPerforceConnection
 {
 public:
 	//This constructor is strictly for internal questions to perforce (get client spec list, etc)
-	FPerforceConnection(const FPerforceConnectionInfo& InConnectionInfo, FPerforceSourceControlProvider& InSCCProvider);
+	// When OutResultInfo is set, errors encountered will be added to it and when not set (nullptr) errors will be logged immediately
+	FPerforceConnection(const FPerforceConnectionInfo& InConnectionInfo, FPerforceSourceControlProvider& InSCCProvider, FSourceControlResultInfo* OutResultInfo);
 	/** API Specific close of source control project*/
 	~FPerforceConnection();
 
@@ -155,8 +156,10 @@ public:
 	/**
 	 * Make a valid connection if possible
 	 * @param InConnectionInfo		Connection credentials
+	 * @param OutResultInfo Optional structure to store errors encountered while establishing the connection.
+	 *                      When set to nullptr, errors will be immediately logged.
 	 */
-	void EstablishConnection(const FPerforceConnectionInfo& InConnectionInfo);
+	void EstablishConnection(const FPerforceConnectionInfo& InConnectionInfo, FSourceControlResultInfo* OutResultInfo);
 
 	/** 
 	 * Return the user of the connection
@@ -227,7 +230,7 @@ public:
 
 private:
 	/** Set up the connection */
-	void Initialize( const FPerforceConnectionInfo& InConnectionInfo, FPerforceSourceControlProvider& SCCProvider);
+	void Initialize(const FPerforceConnectionInfo& InConnectionInfo, FPerforceSourceControlProvider& SCCProvider, FSourceControlResultInfo* OutResultInfo);
 
 private:
 	/** The perforce connection we are using */
diff --git a/Engine/Plugins/Developer/PerforceSourceControl/Source/PerforceSourceControl/Private/PerforceSourceControlProvider.cpp b/Engine/Plugins/Developer/PerforceSourceControl/Source/PerforceSourceControl/Private/PerforceSourceControlProvider.cpp
--- a/Engine/Plugins/Developer/PerforceSourceControl/Source/PerforceSourceControl/Private/PerforceSourceControlProvider.cpp
+++ b/Engine/Plugins/Developer/PerforceSourceControl/Source/PerforceSourceControl/Private/PerforceSourceControlProvider.cpp
@@ -216,21 +216,21 @@ bool FPerforceSourceControlProvider::IsAvailable() const
 	return bServerAvailable && !IsLoginError();
 }
 
-bool FPerforceSourceControlProvider::EstablishPersistentConnection()
+bool FPerforceSourceControlProvider::EstablishPersistentConnection(FSourceControlResultInfo* OutResultInfo)
 {
 	FPerforceConnectionInfo ConnectionInfo = AccessSettings().GetConnectionInfo();
 
 	bool bIsValidConnection = false;
 	if (!PersistentConnection)
 	{
-		PersistentConnection = MakeUnique<FPerforceConnection>(ConnectionInfo, *this);
+		PersistentConnection = MakeUnique<FPerforceConnection>(ConnectionInfo, *this, OutResultInfo);
 	}
 
 	bIsValidConnection = PersistentConnection->IsValidConnection();
 	if (!bIsValidConnection)
 	{
 		PersistentConnection.Reset();
-		PersistentConnection = MakeUnique<FPerforceConnection>(ConnectionInfo, *this);
+		PersistentConnection = MakeUnique<FPerforceConnection>(ConnectionInfo, *this, OutResultInfo);
 		bIsValidConnection = PersistentConnection->IsValidConnection();
 	}
 
@@ -621,9 +621,12 @@ void FPerforceSourceControlProvider::OutputCommandMessages(const FPerforceSource
 	{
 		FMessageLog SourceControlLog("SourceControl");
 
-		for (int32 ErrorIndex = 0; ErrorIndex < InCommand.ResultInfo.ErrorMessages.Num(); ++ErrorIndex)
+		if (!InCommand.Operation->ShouldLogErrors())
 		{
-			SourceControlLog.Error(FText::Format(LOCTEXT("OutputCommandMessagesFormatError", "CommandMessage Command: {0}, Error: {1}"), FText::FromName(InCommand.Operation->GetName()), InCommand.ResultInfo.ErrorMessages[ErrorIndex]));
+			for (int32 ErrorIndex = 0; ErrorIndex < InCommand.ResultInfo.ErrorMessages.Num(); ++ErrorIndex)
+			{
+				SourceControlLog.Error(FText::Format(LOCTEXT("OutputCommandMessagesFormatError", "CommandMessage Command: {0}, Error: {1}"), FText::FromName(InCommand.Operation->GetName()), InCommand.ResultInfo.ErrorMessages[ErrorIndex]));
+			}
 		}
 
 		for (int32 InfoIndex = 0; InfoIndex < InCommand.ResultInfo.InfoMessages.Num(); ++InfoIndex)
@@ -633,9 +636,12 @@ void FPerforceSourceControlProvider::OutputCommandMessages(const FPerforceSource
 	}
 	else // On background threads we must log directly as FMessageLog internals cannot be assumed to be thread safe
 	{
-		for (const FText& Error : InCommand.ResultInfo.ErrorMessages)
+		if (!InCommand.Operation->ShouldLogErrors())
 		{
-			UE_LOG(LogSourceControl, Error, TEXT("Command: %s, Error: %s"), *InCommand.Operation->GetName().ToString(), *Error.ToString());
+			for (const FText& Error : InCommand.ResultInfo.ErrorMessages)
+			{
+				UE_LOG(LogSourceControl, Error, TEXT("Command: %s, Error: %s"), *InCommand.Operation->GetName().ToString(), *Error.ToString());
+			}
 		}
 
 		for (const FText& Info : InCommand.ResultInfo.InfoMessages)
diff --git a/Engine/Plugins/Developer/PerforceSourceControl/Source/PerforceSourceControl/Private/PerforceSourceControlProvider.h b/Engine/Plugins/Developer/PerforceSourceControl/Source/PerforceSourceControl/Private/PerforceSourceControlProvider.h
--- a/Engine/Plugins/Developer/PerforceSourceControl/Source/PerforceSourceControl/Private/PerforceSourceControlProvider.h
+++ b/Engine/Plugins/Developer/PerforceSourceControl/Source/PerforceSourceControl/Private/PerforceSourceControlProvider.h
@@ -93,9 +93,12 @@ public:
 	/**
 	 * Connects to the source control server if the persistent connection is not already established.
 	 *
+	 * @param OutResultInfo Optional structure to store errors encountered while creating the connection.
+	 *                      When set to nullptr, errors will be immediately logged.
+	 *
 	 * @return true if the connection is established or became established and false if the connection failed.
 	 */
-	bool EstablishPersistentConnection();
+	bool EstablishPersistentConnection(FSourceControlResultInfo* OutResultInfo);
 
 	void ResetPersistentConnection();
 
