Editor Perforce integration
- Fallback to uncontrolled changelist if Mark for Add fails.
- Added a debug cvar to prevent commands from running. Simulates a connection issue.

diff --git a/Engine/Plugins/Developer/PerforceSourceControl/Source/PerforceSourceControl/Private/PerforceConnection.cpp b/Engine/Plugins/Developer/PerforceSourceControl/Source/PerforceSourceControl/Private/PerforceConnection.cpp
--- a/Engine/Plugins/Developer/PerforceSourceControl/Source/PerforceSourceControl/Private/PerforceConnection.cpp
+++ b/Engine/Plugins/Developer/PerforceSourceControl/Source/PerforceSourceControl/Private/PerforceConnection.cpp
@@ -3,6 +3,7 @@
 #include "PerforceConnection.h"
 
 #include "HAL/FileManager.h"
+#include "HAL/IConsoleManager.h"
 #include "HAL/PlatformProcess.h"
 #include "HAL/PlatformTime.h"
 #include "ISourceControlModule.h"
@@ -24,6 +25,13 @@ const FString FP4Record::EmptyStr;
 namespace
 {
 
+bool bDebugBreakConnection = false;
+static FAutoConsoleVariableRef CVarDebugBreakConnection(
+	TEXT("SourceControl.Perforce.Debug.BreakConnection"),
+	bDebugBreakConnection,
+	TEXT("Fails any attempt to run a Perforce command."),
+	ECVF_Default);
+
 enum class EP4ClientUserFlags
 {
 	None = 0,
@@ -1139,7 +1147,7 @@ bool FPerforceConnection::RunCommand(	const FString& InCommand, const TArray<FSt
 {
 	TRACE_CPUPROFILER_EVENT_SCOPE_TEXT(*FString::Printf(TEXT("FPerforceConnection::RunCommand_%s"), *InCommand));
 
-	if (!bEstablishedConnection)
+	if (!bEstablishedConnection || bDebugBreakConnection)
 	{
 		return false;
 	}
diff --git a/Engine/Source/Editor/UnrealEd/Private/EditorEngine.cpp b/Engine/Source/Editor/UnrealEd/Private/EditorEngine.cpp
--- a/Engine/Source/Editor/UnrealEd/Private/EditorEngine.cpp
+++ b/Engine/Source/Editor/UnrealEd/Private/EditorEngine.cpp
@@ -4163,15 +4163,20 @@ void UEditorEngine::RunDeferredMarkForAddFiles(bool)
 		return;
 	}
 
+	bool bAddToUncontrolled = true;
+
 	if (ISourceControlModule::Get().IsEnabled())
 	{
 		ISourceControlProvider& SourceControlProvider = ISourceControlModule::Get().GetProvider();
 		if(SourceControlProvider.IsAvailable())
 		{
-			SourceControlProvider.Execute(ISourceControlOperation::Create<FMarkForAdd>(), SourceControlHelpers::PackageFilenames(DeferredFilesToAddToSourceControl));
+			ECommandResult::Type Result =  SourceControlProvider.Execute(ISourceControlOperation::Create<FMarkForAdd>(), SourceControlHelpers::PackageFilenames(DeferredFilesToAddToSourceControl));
+
+			bAddToUncontrolled = (Result == ECommandResult::Failed);
 		}
 	}
-	else if (FUncontrolledChangelistsModule::Get().IsEnabled())
+
+	if (FUncontrolledChangelistsModule::Get().IsEnabled() && bAddToUncontrolled)
 	{
 		FUncontrolledChangelistsModule& UncontrolledChangelistsModule = FUncontrolledChangelistsModule::Get();
 		UncontrolledChangelistsModule.OnNewFilesAdded(DeferredFilesToAddToSourceControl);
