Fix errors from perforce commands not being correctly logged in all scenarios.

#rnx

- Conditions for logging the command were inverted, likely due to that feature being inverted mid way through development (originally it was going to return if things should opt out of logging)

diff --git a/Engine/Plugins/Developer/PerforceSourceControl/Source/PerforceSourceControl/Private/PerforceSourceControlProvider.cpp b/Engine/Plugins/Developer/PerforceSourceControl/Source/PerforceSourceControl/Private/PerforceSourceControlProvider.cpp
--- a/Engine/Plugins/Developer/PerforceSourceControl/Source/PerforceSourceControl/Private/PerforceSourceControlProvider.cpp
+++ b/Engine/Plugins/Developer/PerforceSourceControl/Source/PerforceSourceControl/Private/PerforceSourceControlProvider.cpp
@@ -621,7 +621,7 @@ void FPerforceSourceControlProvider::OutputCommandMessages(const FPerforceSource
 	{
 		FMessageLog SourceControlLog("SourceControl");
 
-		if (!InCommand.Operation->ShouldLogErrors())
+		if (InCommand.Operation->ShouldLogErrors())
 		{
 			for (int32 ErrorIndex = 0; ErrorIndex < InCommand.ResultInfo.ErrorMessages.Num(); ++ErrorIndex)
 			{
@@ -636,7 +636,7 @@ void FPerforceSourceControlProvider::OutputCommandMessages(const FPerforceSource
 	}
 	else // On background threads we must log directly as FMessageLog internals cannot be assumed to be thread safe
 	{
-		if (!InCommand.Operation->ShouldLogErrors())
+		if (InCommand.Operation->ShouldLogErrors())
 		{
 			for (const FText& Error : InCommand.ResultInfo.ErrorMessages)
 			{
