FPerforceSourceControlProvider : IsAvailable() now considers whether or not the current workspace maps the current project directory.

bIsWorkspaceValidForProject defaults to true and only gets cleared in FPerforceConnectWorker::Execute if we explicitly check this condition ( CheckWorkspaceRecordSet ) and it fails.

When FPerforceConnectWorker::Execute early outs due to an empty workspace it should set bCommandSuccessful to false ( this would otherwise be the result in that case if we didn't early out ).

FSourceControlModule : In OnSourceControlDialogClosed fix a bug where closing the SSourceControlLogin dialog with the close button would lead to incorrect revision control status.

diff --git a/Engine/Plugins/Developer/PerforceSourceControl/Source/PerforceSourceControl/Private/PerforceSourceControlOperations.cpp b/Engine/Plugins/Developer/PerforceSourceControl/Source/PerforceSourceControl/Private/PerforceSourceControlOperations.cpp
--- a/Engine/Plugins/Developer/PerforceSourceControl/Source/PerforceSourceControl/Private/PerforceSourceControlOperations.cpp
+++ b/Engine/Plugins/Developer/PerforceSourceControl/Source/PerforceSourceControl/Private/PerforceSourceControlOperations.cpp
@@ -433,9 +433,9 @@ bool FPerforceConnectWorker::Execute(FPerforceSourceControlCommand& InCommand)
 		// Running FConnect for perforce does not actually connect to the server but
 		// validates that the clientspec for the connection is correct and working.
 		// So if the connection does not have a valid clientspec then we might as well
-		// early out at this point and return true.
+		// early out at this point and return false.
 
-		InCommand.bCommandSuccessful = true;
+		InCommand.bCommandSuccessful = false;
 		return InCommand.bCommandSuccessful;
 	}
 
@@ -468,6 +468,8 @@ bool FPerforceConnectWorker::Execute(FPerforceSourceControlCommand& InCommand)
 				TSharedRef<FConnect, ESPMode::ThreadSafe> Operation = StaticCastSharedRef<FConnect>(InCommand.Operation);
 				Operation->SetErrorText(Notification);
 			}
+			
+			GetSCCProvider().SetIsWorkspaceValidForProject(InCommand.bCommandSuccessful);
 		}
 
 		if(InCommand.bCommandSuccessful)
diff --git a/Engine/Plugins/Developer/PerforceSourceControl/Source/PerforceSourceControl/Private/PerforceSourceControlProvider.cpp b/Engine/Plugins/Developer/PerforceSourceControl/Source/PerforceSourceControl/Private/PerforceSourceControlProvider.cpp
--- a/Engine/Plugins/Developer/PerforceSourceControl/Source/PerforceSourceControl/Private/PerforceSourceControlProvider.cpp
+++ b/Engine/Plugins/Developer/PerforceSourceControl/Source/PerforceSourceControl/Private/PerforceSourceControlProvider.cpp
@@ -69,6 +69,7 @@ FPerforceSourceControlProvider::FPerforceSourceControlProvider()
 	, InitialSettings(FSourceControlInitSettings::EBehavior::OverrideExisting, FSourceControlInitSettings::ECmdLineFlags::ReadAll)
 	, OwnerName(TEXT("Default"))
 	, bServerAvailable(false)
+	, bIsWorkspaceValidForProject(true)
 	, bLoginError(false)
 	, PersistentConnection(nullptr)
 {
@@ -82,6 +83,7 @@ FPerforceSourceControlProvider::FPerforceSourceControlProvider(const FStringView
 	, InitialSettings(InInitialSettings)
 	, OwnerName(InOwnerName)
 	, bServerAvailable(false)
+	, bIsWorkspaceValidForProject(true)
 	, bLoginError(false)
 	, PersistentConnection(nullptr)
 {
@@ -213,7 +215,7 @@ bool FPerforceSourceControlProvider::IsEnabled() const
 
 bool FPerforceSourceControlProvider::IsAvailable() const
 {
-	return bServerAvailable && !IsLoginError();
+	return bServerAvailable && !IsLoginError() && bIsWorkspaceValidForProject;
 }
 
 bool FPerforceSourceControlProvider::EstablishPersistentConnection(FSourceControlResultInfo* OutResultInfo)
@@ -370,6 +372,11 @@ bool FPerforceSourceControlProvider::IsLoginError() const
 	return bLoginError;
 }
 
+void FPerforceSourceControlProvider::SetIsWorkspaceValidForProject(const bool bInIsWorkspaceValidForProject)
+{
+	bIsWorkspaceValidForProject = bInIsWorkspaceValidForProject;
+}
+
 TArray<FText> FPerforceSourceControlProvider::GetLastErrors() const
 {
 	FScopeLock Lock(&LastErrorsCriticalSection);
diff --git a/Engine/Plugins/Developer/PerforceSourceControl/Source/PerforceSourceControl/Private/PerforceSourceControlProvider.h b/Engine/Plugins/Developer/PerforceSourceControl/Source/PerforceSourceControl/Private/PerforceSourceControlProvider.h
--- a/Engine/Plugins/Developer/PerforceSourceControl/Source/PerforceSourceControl/Private/PerforceSourceControlProvider.h
+++ b/Engine/Plugins/Developer/PerforceSourceControl/Source/PerforceSourceControl/Private/PerforceSourceControlProvider.h
@@ -78,6 +78,9 @@ public:
 	/** Did most recent command generate a login error */
 	bool IsLoginError() const;
 
+	/** Sets whether or not the current workspace maps to the current project directory */
+	void SetIsWorkspaceValidForProject(const bool bInIsWorkspaceValidForProject);
+
 	/** Get list of error messages that occurred after last perforce command */
 	TArray<FText> GetLastErrors() const;
 
@@ -175,6 +178,9 @@ private:
 	/** Indicates if source control integration is available or not. */
 	bool bServerAvailable;
 
+	/** Whether or not the current workspace maps to the current project directory; this is true unless explicitly cleared */
+	TAtomic<bool> bIsWorkspaceValidForProject;
+
 	/** Saw login error when running last command. */
 	TAtomic<bool> bLoginError;
 
diff --git a/Engine/Source/Developer/SourceControl/Private/SSourceControlLogin.h b/Engine/Source/Developer/SourceControl/Private/SSourceControlLogin.h
--- a/Engine/Source/Developer/SourceControl/Private/SSourceControlLogin.h
+++ b/Engine/Source/Developer/SourceControl/Private/SSourceControlLogin.h
@@ -50,6 +50,11 @@ public:
 	 */
 	void RefreshSettings();
 
+	bool IsConnected() const
+	{
+		return ConnectionState == ELoginConnectionState::Connected;
+	}
+
 private:
 
 	/** Ticks the source control module (only necessary if the login window is modal) */
diff --git a/Engine/Source/Developer/SourceControl/Private/SourceControlModule.cpp b/Engine/Source/Developer/SourceControl/Private/SourceControlModule.cpp
--- a/Engine/Source/Developer/SourceControl/Private/SourceControlModule.cpp
+++ b/Engine/Source/Developer/SourceControl/Private/SourceControlModule.cpp
@@ -221,6 +221,13 @@ void FSourceControlModule::ShowLoginDialog(const FSourceControlLoginClosed& InOn
 
 void FSourceControlModule::OnSourceControlDialogClosed(const TSharedRef<SWindow>& InWindow)
 {
+#if SOURCE_CONTROL_WITH_SLATE
+	if (!SourceControlLoginPtr->IsConnected())
+	{
+		SetProvider("None");
+	}
+#endif
+	
 	SourceControlLoginWindowPtr = NULL;
 	SourceControlLoginPtr = NULL;
 	--TemporarilyDisabledCount;
