Perforce Provider: Guard access to the state caches for concurrent queries to files & changelists states.
Fixes a crash when right clicking on the tree view menus in the content browser while the thumbnails are fetching their Perforce states.

diff --git a/Engine/Plugins/Developer/PerforceSourceControl/Source/PerforceSourceControl/Private/PerforceSourceControlProvider.cpp b/Engine/Plugins/Developer/PerforceSourceControl/Source/PerforceSourceControl/Private/PerforceSourceControlProvider.cpp
--- a/Engine/Plugins/Developer/PerforceSourceControl/Source/PerforceSourceControl/Private/PerforceSourceControlProvider.cpp
+++ b/Engine/Plugins/Developer/PerforceSourceControl/Source/PerforceSourceControl/Private/PerforceSourceControlProvider.cpp
@@ -11,6 +11,7 @@
 #include "Misc/CommandLine.h"
 #include "Misc/MessageDialog.h"
 #include "Misc/QueuedThreadPool.h"
+#include "Misc/ScopeRWLock.h"
 #include "PerforceConnection.h"
 #include "PerforceConnectionInfo.h"
 #include "PerforceSourceControlCommand.h"
@@ -120,14 +121,22 @@ void FPerforceSourceControlProvider::Close()
 	}
 
 	// clear the cache
-	StateCache.Empty();
+	{
+		FWriteScopeLock WriteToStateCacheLock(StateCacheLock);
+		StateCache.Empty();
+	}
 
 	bServerAvailable = false;
 }
 
 TSharedRef<FPerforceSourceControlState, ESPMode::ThreadSafe> FPerforceSourceControlProvider::GetStateInternal(const FString& Filename)
 {
-	TSharedRef<FPerforceSourceControlState, ESPMode::ThreadSafe>* State = StateCache.Find(Filename);
+	TSharedRef<FPerforceSourceControlState, ESPMode::ThreadSafe>* State;
+	{
+		FReadScopeLock ReadFromStateCacheLock(StateCacheLock);
+		State = StateCache.Find(Filename);
+	}
+
 	if(State != NULL)
 	{
 		// found cached item
@@ -135,6 +144,8 @@ TSharedRef<FPerforceSourceControlState, ESPMode::ThreadSafe> FPerforceSourceCont
 	}
 	else
 	{
+		FWriteScopeLock WriteToStateCacheLock(StateCacheLock);
+
 		// cache an unknown state for this item
 		TSharedRef<FPerforceSourceControlState, ESPMode::ThreadSafe> NewState = MakeShared<FPerforceSourceControlState>(Filename);
 		StateCache.Add(Filename, NewState);
@@ -144,7 +155,12 @@ TSharedRef<FPerforceSourceControlState, ESPMode::ThreadSafe> FPerforceSourceCont
 
 TSharedRef<FPerforceSourceControlChangelistState, ESPMode::ThreadSafe> FPerforceSourceControlProvider::GetStateInternal(const FPerforceSourceControlChangelist& InChangelist)
 {
-	TSharedRef<FPerforceSourceControlChangelistState, ESPMode::ThreadSafe>* State = ChangelistsStateCache.Find(InChangelist);
+	TSharedRef<FPerforceSourceControlChangelistState, ESPMode::ThreadSafe>* State;
+	{
+		FReadScopeLock ReadFromChangelistsStateCacheLock(ChangelistsStateCacheLock);
+		State = ChangelistsStateCache.Find(InChangelist);
+	}
+
 	if (State != NULL)
 	{
 		// found cached item
@@ -152,6 +168,8 @@ TSharedRef<FPerforceSourceControlChangelistState, ESPMode::ThreadSafe> FPerforce
 	}
 	else
 	{
+		FWriteScopeLock WriteToChangelistsStateCacheLock(ChangelistsStateCacheLock);
+
 		// cache an unknown state for this item
 		TSharedRef<FPerforceSourceControlChangelistState, ESPMode::ThreadSafe> NewState = MakeShared<FPerforceSourceControlChangelistState>(InChangelist);
 		ChangelistsStateCache.Add(InChangelist, NewState);
@@ -438,16 +456,20 @@ ECommandResult::Type FPerforceSourceControlProvider::GetState(const TArray<FSour
 
 bool FPerforceSourceControlProvider::RemoveFileFromCache(const FString& Filename)
 {
+	FWriteScopeLock WriteToStateCacheLock(StateCacheLock);
 	return StateCache.Remove(Filename) > 0;
 }
 
 bool FPerforceSourceControlProvider::RemoveChangelistFromCache(const FPerforceSourceControlChangelist& Changelist)
 {
+	FWriteScopeLock WriteToChangelistsStateCacheLock(ChangelistsStateCacheLock);
 	return ChangelistsStateCache.Remove(Changelist) > 0;
 }
 
 TArray<FSourceControlStateRef> FPerforceSourceControlProvider::GetCachedStateByPredicate(TFunctionRef<bool(const FSourceControlStateRef&)> Predicate) const
 {
+	FReadScopeLock ReadFromStateCacheLock(StateCacheLock);
+
 	TArray<FSourceControlStateRef> Result;
 	for (const auto& CacheItem : StateCache)
 	{
@@ -462,6 +484,8 @@ TArray<FSourceControlStateRef> FPerforceSourceControlProvider::GetCachedStateByP
 
 TArray<FSourceControlChangelistStateRef> FPerforceSourceControlProvider::GetCachedStateByPredicate(TFunctionRef<bool(const FSourceControlChangelistStateRef&)> Predicate) const
 {
+	FReadScopeLock ReadFromChangelistsStateCacheLock(ChangelistsStateCacheLock);
+
 	TArray<FSourceControlChangelistStateRef> Result;
 	for (const auto& CacheItem : ChangelistsStateCache)
 	{
@@ -797,10 +821,13 @@ TArray<FSourceControlChangelistRef> FPerforceSourceControlProvider::GetChangelis
 	}
 
 	TArray<FSourceControlChangelistRef> Changelists;
-	Algo::Transform(ChangelistsStateCache, Changelists, [](const TPair<FPerforceSourceControlChangelist, TSharedRef<FPerforceSourceControlChangelistState, ESPMode::ThreadSafe>>& Pair)
 	{
-		return MakeShared<FPerforceSourceControlChangelist, ESPMode::ThreadSafe>(Pair.Key);
-	});
+		FReadScopeLock ReadFromChangelistsStateCacheLock(ChangelistsStateCacheLock);
+		Algo::Transform(ChangelistsStateCache, Changelists, [](const TPair<FPerforceSourceControlChangelist, TSharedRef<FPerforceSourceControlChangelistState, ESPMode::ThreadSafe>>& Pair)
+		{
+			return MakeShared<FPerforceSourceControlChangelist, ESPMode::ThreadSafe>(Pair.Key);
+		});
+	}
 
 	// NOTE: Sort in ascending number. If this behavior needs to be configurable, we could have 3-state enum param: 'default, ascending, descending'. For P4, the 'default'
 	//       should be ascending. The changelists are source control agnostics, so sorting a changelist files is probably not a notion that should be leaked in the generic interface.
diff --git a/Engine/Plugins/Developer/PerforceSourceControl/Source/PerforceSourceControl/Private/PerforceSourceControlProvider.h b/Engine/Plugins/Developer/PerforceSourceControl/Source/PerforceSourceControl/Private/PerforceSourceControlProvider.h
--- a/Engine/Plugins/Developer/PerforceSourceControl/Source/PerforceSourceControl/Private/PerforceSourceControlProvider.h
+++ b/Engine/Plugins/Developer/PerforceSourceControl/Source/PerforceSourceControl/Private/PerforceSourceControlProvider.h
@@ -12,6 +12,8 @@
 #include "PerforceSourceControlState.h"
 #include "SourceControlInitSettings.h"
 
+#include "HAL/CriticalSection.h"
+
 class FPerforceConnection;
 class FPerforceSourceControlCommand;
 
@@ -195,7 +197,10 @@ private:
 
 	/** State cache */
 	TMap<FString, TSharedRef<class FPerforceSourceControlState, ESPMode::ThreadSafe> > StateCache;
+	mutable FRWLock StateCacheLock;
+
 	TMap<FPerforceSourceControlChangelist, TSharedRef<class FPerforceSourceControlChangelistState, ESPMode::ThreadSafe> > ChangelistsStateCache;
+	mutable FRWLock ChangelistsStateCacheLock;
 
 	/** Queue for commands given by the main thread */
 	TArray <FPerforceSourceControlCommand* > CommandQueue;
