Prevent dumping Perforce client info every time an asset is opened

Note: this mirrors for info messages what we already support for error messages.

* Add support to ISourceControlOperation and FSourceControlOperationBase to control whether info messages from Perforce should be logged to the Unreal console.
* Inhibit info messages in FSourceControlMenuHelpers::CheckSourceControlStatus to prevent logging the whole client info each time source control status is checked (for example when the status bar is opened).

diff --git a/Engine/Plugins/Developer/PerforceSourceControl/Source/PerforceSourceControl/Private/PerforceSourceControlProvider.cpp b/Engine/Plugins/Developer/PerforceSourceControl/Source/PerforceSourceControl/Private/PerforceSourceControlProvider.cpp
--- a/Engine/Plugins/Developer/PerforceSourceControl/Source/PerforceSourceControl/Private/PerforceSourceControlProvider.cpp
+++ b/Engine/Plugins/Developer/PerforceSourceControl/Source/PerforceSourceControl/Private/PerforceSourceControlProvider.cpp
@@ -660,9 +660,16 @@ void FPerforceSourceControlProvider::OutputCommandMessages(const FPerforceSource
 			}
 		}
 
-		for (int32 InfoIndex = 0; InfoIndex < InCommand.ResultInfo.InfoMessages.Num(); ++InfoIndex)
+		if (InCommand.Operation->ShouldLogInfos())
 		{
-			SourceControlLog.Info(FText::Format(LOCTEXT("OutputCommandMessagesFormatInfo", "CommandMessage Command: {0}, Info: {1}"), FText::FromName(InCommand.Operation->GetName()), InCommand.ResultInfo.InfoMessages[InfoIndex]));
+			for (int32 InfoIndex = 0; InfoIndex < InCommand.ResultInfo.InfoMessages.Num(); ++InfoIndex)
+			{
+				SourceControlLog.Info(FText::Format(
+					LOCTEXT("OutputCommandMessagesFormatInfo", "CommandMessage Command: {0}, Info: {1}"),
+					FText::FromName(InCommand.Operation->GetName()),
+					InCommand.ResultInfo.InfoMessages[InfoIndex]
+				));
+			}
 		}
 	}
 	else // On background threads we must log directly as FMessageLog internals cannot be assumed to be thread safe
diff --git a/Engine/Source/Developer/SourceControl/Public/ISourceControlOperation.h b/Engine/Source/Developer/SourceControl/Public/ISourceControlOperation.h
--- a/Engine/Source/Developer/SourceControl/Public/ISourceControlOperation.h
+++ b/Engine/Source/Developer/SourceControl/Public/ISourceControlOperation.h
@@ -75,6 +75,14 @@ public:
 		return true;
 	}
 
+	/**
+	 * Returns true when the operation is allowed to log info messages immediately when they are encountered.
+	 */
+	virtual bool ShouldLogInfos() const
+	{
+		return true;
+	}
+
 	/** Factory method for easier operation creation */
 	template<typename Type, typename... TArgs>
 	static TSharedRef<Type, ESPMode::ThreadSafe> Create(TArgs&&... Args)
diff --git a/Engine/Source/Developer/SourceControl/Public/SourceControlOperationBase.h b/Engine/Source/Developer/SourceControl/Public/SourceControlOperationBase.h
--- a/Engine/Source/Developer/SourceControl/Public/SourceControlOperationBase.h
+++ b/Engine/Source/Developer/SourceControl/Public/SourceControlOperationBase.h
@@ -12,8 +12,10 @@ public:
 	{
 		/* No specialization applied */
 		None = 0,
-		/* When set the operation should not log any error but should continue to store them in the ResultInfo member */
-		DisableErrorLogging = 1 << 0
+		/* When set the operation should not log any errors but should continue to store them in the ResultInfo member */
+		DisableErrorLogging = 1 << 0,
+		/* When set the operation should not log any info messages but should continue to store them in the ResultInfo member */
+		DisableInfoLogging = 1 << 1,
 	};
 	FRIEND_ENUM_CLASS_FLAGS(EFlags)
 
@@ -70,6 +72,23 @@ public:
 		return !EnumHasAllFlags(Flags, EFlags::DisableErrorLogging);
 	}
 
+	void SetEnableInfoLogging(bool bEnabled)
+	{
+		if (bEnabled)
+		{
+			EnumRemoveFlags(Flags, EFlags::DisableInfoLogging);
+		}
+		else
+		{
+			EnumAddFlags(Flags, EFlags::DisableInfoLogging);
+		}
+	}
+
+	virtual bool ShouldLogInfos() const override
+	{
+		return !EnumHasAllFlags(Flags, EFlags::DisableInfoLogging);
+	}
+
 	FSourceControlResultInfo ResultInfo;
 
 protected:
diff --git a/Engine/Source/Editor/StatusBar/Private/SourceControlMenuHelpers.cpp b/Engine/Source/Editor/StatusBar/Private/SourceControlMenuHelpers.cpp
--- a/Engine/Source/Editor/StatusBar/Private/SourceControlMenuHelpers.cpp
+++ b/Engine/Source/Editor/StatusBar/Private/SourceControlMenuHelpers.cpp
@@ -227,9 +227,14 @@ void FSourceControlMenuHelpers::CheckSourceControlStatus()
 	ISourceControlModule& SourceControlModule = ISourceControlModule::Get();
 	if (SourceControlModule.IsEnabled())
 	{
-		SourceControlModule.GetProvider().Execute(ISourceControlOperation::Create<FConnect>(),
+		TSharedRef<FConnect> Operation = ISourceControlOperation::Create<FConnect>();
+		// Disable logging of info messages to prevent dumping 100+ rows of client info to the log every time an asset editor is opened.
+		Operation->SetEnableInfoLogging(false);
+		SourceControlModule.GetProvider().Execute(
+			Operation,
 			EConcurrency::Asynchronous,
-			FSourceControlOperationComplete::CreateStatic(&FSourceControlMenuHelpers::OnSourceControlOperationComplete));
+			FSourceControlOperationComplete::CreateStatic(&FSourceControlMenuHelpers::OnSourceControlOperationComplete)
+		);
 		QueryState = EQueryState::Querying;
 	}
 }
