Fix a bug in CheckWorkspaceRecordSet where it would wrongly return false and log an error for the case where the ApplicationPath is equal to the Root except that the Root has a trailing '/'

This situation occurs if the project file being opened is directly in the workspace root directory (not in a subdirectory).

The fix uses FPathViews::IsParentPathOf which returns true in this case:
"Presence or absence of terminating separator (/) is ignored in the comparison."

The same bug is being fixed also in FPerforceConnection::GetWorkspaceList.

diff --git a/Engine/Plugins/Developer/PerforceSourceControl/Source/PerforceSourceControl/Private/PerforceConnection.cpp b/Engine/Plugins/Developer/PerforceSourceControl/Source/PerforceSourceControl/Private/PerforceConnection.cpp
--- a/Engine/Plugins/Developer/PerforceSourceControl/Source/PerforceSourceControl/Private/PerforceConnection.cpp
+++ b/Engine/Plugins/Developer/PerforceSourceControl/Source/PerforceSourceControl/Private/PerforceConnection.cpp
@@ -10,6 +10,7 @@
 #include "Logging/MessageLog.h"
 #include "Memory/SharedBuffer.h"
 #include "Misc/Paths.h"
+#include "Misc/PathViews.h"
 #include "PerforceMessageLog.h"
 #include "PerforceSourceControlProvider.h"
 #include "PerforceSourceControlSettings.h"
@@ -1091,7 +1092,7 @@ bool FPerforceConnection::GetWorkspaceList(const FPerforceConnectionInfo& InConn
 					}
 
 					// Only allow paths that are ACTUALLY legit for this application
-					if (bIsNullClientRootPath || ApplicationPath.Contains(ClientRootPath) )
+					if (bIsNullClientRootPath || FPathViews::IsParentPathOf(ClientRootPath, ApplicationPath))
 					{
 						OutWorkspaceList.Add(ClientName);
 					}
diff --git a/Engine/Plugins/Developer/PerforceSourceControl/Source/PerforceSourceControl/Private/PerforceSourceControlOperations.cpp b/Engine/Plugins/Developer/PerforceSourceControl/Source/PerforceSourceControl/Private/PerforceSourceControlOperations.cpp
--- a/Engine/Plugins/Developer/PerforceSourceControl/Source/PerforceSourceControl/Private/PerforceSourceControlOperations.cpp
+++ b/Engine/Plugins/Developer/PerforceSourceControl/Source/PerforceSourceControl/Private/PerforceSourceControlOperations.cpp
@@ -11,6 +11,7 @@
 #include "ISourceControlModule.h"
 #include "Misc/EngineVersion.h"
 #include "Misc/Paths.h"
+#include "Misc/PathViews.h"
 #include "PerforceConnection.h"
 #include "PerforceMessageLog.h"
 #include "PerforceSourceControlChangeStatusOperation.h"
@@ -385,14 +386,14 @@ static bool CheckWorkspaceRecordSet(const FP4RecordSet& InRecords, TArray<FText>
 			Root += TEXT("/");
 		}
 
-		if (bIsNullClientRootPath || ApplicationPath.Contains(Root))
+		if (bIsNullClientRootPath || FPathViews::IsParentPathOf(Root, ApplicationPath))
 		{
 			return true;
 		}
 		else
 		{
 			const FString& Client = Record(TEXT("Client"));
-			OutNotificationText = FText::Format(LOCTEXT("WorkspaceError", "Workspace '{0}' does not map into this project's directory '{1}'"), FText::FromString(Client), FText::FromString(ApplicationPath));
+			OutNotificationText = FText::Format(LOCTEXT("WorkspaceError", "Workspace '{0}' with root '{1}' does not map into this project's directory '{2}'"), FText::FromString(Client), FText::FromString(Root), FText::FromString(ApplicationPath));
 			OutErrorMessages.Add(OutNotificationText);
 			OutErrorMessages.Add(LOCTEXT("WorkspaceHelp", "You should set your workspace up to map to a directory at or above the project's directory."));
 		}
