SourceControl / Perforce: Fixed 'not in client view' errors reported for cooked content.

#rnx

diff --git a/Engine/Plugins/Developer/PerforceSourceControl/Source/PerforceSourceControl/Private/PerforceSourceControlOperations.cpp b/Engine/Plugins/Developer/PerforceSourceControl/Source/PerforceSourceControl/Private/PerforceSourceControlOperations.cpp
--- a/Engine/Plugins/Developer/PerforceSourceControl/Source/PerforceSourceControl/Private/PerforceSourceControlOperations.cpp
+++ b/Engine/Plugins/Developer/PerforceSourceControl/Source/PerforceSourceControl/Private/PerforceSourceControlOperations.cpp
@@ -1541,6 +1541,18 @@ static void ParseUpdateStatusResults(const FP4RecordSet& InRecords, const TArray
 			FPerforceSourceControlState& State = OutStates.Last();
 			State.State = EPerforceState::NotUnderClientRoot;
 		}
+		
+		//@todo P4 could be returning localized error messages
+		const int32 NotInClientViewPos = Error.ToString().Find(TEXT(" - file(s) not in client view"), ESearchCase::IgnoreCase, ESearchDir::FromStart);
+		if(NotInClientViewPos != INDEX_NONE)
+		{
+			// found an error about a file that is not in the client view, treating it same as not under client root
+			FString FullPath(Error.ToString().Left(NotInClientViewPos));
+			FPaths::NormalizeFilename(FullPath);
+			OutStates.Add(FPerforceSourceControlState(FullPath));
+			FPerforceSourceControlState& State = OutStates.Last();
+			State.State = EPerforceState::NotUnderClientRoot;
+		}
 	}
 }
 
@@ -1959,6 +1971,7 @@ bool FPerforceUpdateStatusWorker::Execute(FPerforceSourceControlCommand& InComma
 			ParseUpdateStatusResults(Records, InCommand.ResultInfo.ErrorMessages, OutStates, ContentRoot, BranchModifications);
 			RemoveRedundantErrors(InCommand, TEXT(" - no such file(s)."), false);
 			RemoveRedundantErrors(InCommand, TEXT("' is not under client's root '"));
+			RemoveRedundantErrors(InCommand, TEXT(" - file(s) not in client view"), false);
 			RemoveRedundantErrors(InCommand, TEXT(" - protected namespace - access denied"), false);
 		}
 		else
diff --git a/Engine/Plugins/Developer/PerforceSourceControl/Source/PerforceSourceControl/Private/PerforceSourceControlState.h b/Engine/Plugins/Developer/PerforceSourceControl/Source/PerforceSourceControl/Private/PerforceSourceControlState.h
--- a/Engine/Plugins/Developer/PerforceSourceControl/Source/PerforceSourceControl/Private/PerforceSourceControlState.h
+++ b/Engine/Plugins/Developer/PerforceSourceControl/Source/PerforceSourceControl/Private/PerforceSourceControlState.h
@@ -38,7 +38,7 @@ namespace EPerforceState
 		/** File is marked for delete */
 		MarkedForDelete	= 8,
 
-		/** Not under client root */
+		/** Not under client root or in client view */
 		NotUnderClientRoot	= 9,
 
 		/** Opened for branch */
