Re-enable p4 command logging, we now only supress the problematic command that was creating log spam.

- Revert the change to the PerforceConnection log verbosity made in CL 39552774.
- Instead we now check the current verbosity of LogSourceControl and use that to decide if our 'p4 fstat' commands should be run with ERunCommandFlags::DisableCommandLogging or not.
- It is currently set up so that we only log info about the fstat command if LogSourceControl is currently set to 'verbose'

diff --git a/Engine/Plugins/Developer/PerforceSourceControl/Source/PerforceSourceControl/Private/PerforceConnection.cpp b/Engine/Plugins/Developer/PerforceSourceControl/Source/PerforceSourceControl/Private/PerforceConnection.cpp
--- a/Engine/Plugins/Developer/PerforceSourceControl/Source/PerforceSourceControl/Private/PerforceConnection.cpp
+++ b/Engine/Plugins/Developer/PerforceSourceControl/Source/PerforceSourceControl/Private/PerforceConnection.cpp
@@ -1183,7 +1183,7 @@ bool FPerforceConnection::RunCommand(	const FString& InCommand, const TArray<FSt
 
 	if (bLogCommandDetails)
 	{
-		UE_LOG( LogSourceControl, Verbose, TEXT("Attempting 'p4 %s'"), *FullCommand );
+		UE_LOG( LogSourceControl, Log, TEXT("Attempting 'p4 %s'"), *FullCommand );
 	}
 
 	double SCCStartTime = FPlatformTime::Seconds();
@@ -1236,7 +1236,7 @@ bool FPerforceConnection::RunCommand(	const FString& InCommand, const TArray<FSt
 		const double ExecutionTime = FPlatformTime::Seconds() - SCCStartTime;
 		if (ExecutionTime >= 0.1)
 		{
-			UE_LOG(LogSourceControl, Verbose, TEXT("P4 execution time: %0.4f seconds. Command: %s"), ExecutionTime, *FullCommand);
+			UE_LOG(LogSourceControl, Log, TEXT("P4 execution time: %0.4f seconds. Command: %s"), ExecutionTime, *FullCommand);
 		}
 		else
 		{
diff --git a/Engine/Plugins/Developer/PerforceSourceControl/Source/PerforceSourceControl/Private/PerforceSourceControlOperations.cpp b/Engine/Plugins/Developer/PerforceSourceControl/Source/PerforceSourceControl/Private/PerforceSourceControlOperations.cpp
--- a/Engine/Plugins/Developer/PerforceSourceControl/Source/PerforceSourceControl/Private/PerforceSourceControlOperations.cpp
+++ b/Engine/Plugins/Developer/PerforceSourceControl/Source/PerforceSourceControl/Private/PerforceSourceControlOperations.cpp
@@ -1949,6 +1949,16 @@ bool FPerforceUpdateStatusWorker::Execute(FPerforceSourceControlCommand& InComma
 			// Initially successful
 			InCommand.bCommandSuccessful = true;
 
+			ERunCommandFlags Flags = ERunCommandFlags::Default;
+#if !NO_LOGGING
+			// FPerforceConnect::RunCommand can log a lot of spam when it comes to fstat which is determintal to editor performance.
+			// So we want to supress that log output unless the user is running with LogSourceControl as verbose.
+			if (LogSourceControl.IsSuppressed(ELogVerbosity::Verbose))
+			{
+				Flags |= ERunCommandFlags::DisableCommandLogging;
+			}
+#endif // !NO_LOGGING
+
 			// Parse branch modifications
 			TMap<FString, FBranchModification> BranchModifications;
 			if (StatusBranches.Num())
@@ -1961,12 +1971,12 @@ bool FPerforceUpdateStatusWorker::Execute(FPerforceSourceControlCommand& InComma
 				RevisionParameters.Insert(TEXT("-Of"), 0);
 
 				FP4RecordSet RevisionRecords;
-				InCommand.bCommandSuccessful &= Connection.RunCommand(TEXT("fstat"), RevisionParameters, RevisionRecords, InCommand.ResultInfo, FOnIsCancelled::CreateRaw(&InCommand, &FPerforceSourceControlCommand::IsCanceled), InCommand.bConnectionDropped);
+				InCommand.bCommandSuccessful &= Connection.RunCommand(TEXT("fstat"), RevisionParameters, RevisionRecords, nullptr, InCommand.ResultInfo, FOnIsCancelled::CreateRaw(&InCommand, &FPerforceSourceControlCommand::IsCanceled), InCommand.bConnectionDropped, Flags);
 				ParseBranchModificationResults(RevisionRecords, InCommand.ResultInfo.ErrorMessages, ContentRoot, BranchModifications);
 			}
 
 			FP4RecordSet Records;
-			InCommand.bCommandSuccessful &= Connection.RunCommand(TEXT("fstat"), Parameters, Records, InCommand.ResultInfo, FOnIsCancelled::CreateRaw(&InCommand, &FPerforceSourceControlCommand::IsCanceled), InCommand.bConnectionDropped);
+			InCommand.bCommandSuccessful &= Connection.RunCommand(TEXT("fstat"), Parameters, Records, nullptr, InCommand.ResultInfo, FOnIsCancelled::CreateRaw(&InCommand, &FPerforceSourceControlCommand::IsCanceled), InCommand.bConnectionDropped, Flags);
 			ParseUpdateStatusResults(Records, InCommand.ResultInfo.ErrorMessages, OutStates, ContentRoot, BranchModifications);
 			RemoveRedundantErrors(InCommand, TEXT(" - no such file(s)."), false);
 			RemoveRedundantErrors(InCommand, TEXT("' is not under client's root '"));
