[Source Control] Fix an issue where branched files on top of a delete was not functioning

#rnx

diff --git a/Engine/Plugins/Developer/PerforceSourceControl/Source/PerforceSourceControl/Private/PerforceSourceControlOperations.cpp b/Engine/Plugins/Developer/PerforceSourceControl/Source/PerforceSourceControl/Private/PerforceSourceControlOperations.cpp
--- a/Engine/Plugins/Developer/PerforceSourceControl/Source/PerforceSourceControl/Private/PerforceSourceControlOperations.cpp
+++ b/Engine/Plugins/Developer/PerforceSourceControl/Source/PerforceSourceControl/Private/PerforceSourceControlOperations.cpp
@@ -2640,7 +2640,15 @@ bool FPerforceCopyWorker::Execute(FPerforceSourceControlCommand& InCommand)
 			if (InCommand.bCommandSuccessful)
 			{
 				TArray<FString> ResolveParameters;
-				ResolveParameters.Add(TEXT("-ay"));	// 'accept yours'
+				switch (Operation->ResolveMethod)
+				{
+					case FCopy::EResolveMethod::AcceptTarget:
+						ResolveParameters.Add(TEXT("-ay"));	// 'accept yours' to leave the target unchanged
+					break;
+					case FCopy::EResolveMethod::AcceptSource:
+						ResolveParameters.Add(TEXT("-at")); // 'accept theirs' to allow branch over deleted files
+					break;
+				}
 				ResolveParameters.Add(MoveTemp(DestinationPath));
 				InCommand.bCommandSuccessful = Connection.RunCommand(TEXT("resolve"), ResolveParameters, Records, InCommand.ResultInfo, FOnIsCancelled::CreateRaw(&InCommand, &FPerforceSourceControlCommand::IsCanceled), InCommand.bConnectionDropped);
 			}
diff --git a/Engine/Plugins/Developer/PerforceSourceControl/Source/PerforceSourceControl/Private/PerforceSourceControlState.cpp b/Engine/Plugins/Developer/PerforceSourceControl/Source/PerforceSourceControl/Private/PerforceSourceControlState.cpp
--- a/Engine/Plugins/Developer/PerforceSourceControl/Source/PerforceSourceControl/Private/PerforceSourceControlState.cpp
+++ b/Engine/Plugins/Developer/PerforceSourceControl/Source/PerforceSourceControl/Private/PerforceSourceControlState.cpp
@@ -269,12 +269,12 @@ bool FPerforceSourceControlState::CanCheckout() const
 	bool bCanDoCheckout = false;
 
 	const bool bIsInP4NotCheckedOut = State == EPerforceState::ReadOnly;
+	const bool bIsOpenForBranch = State == EPerforceState::Branched;
 	if (!bBinary && !bExclusiveCheckout)
 	{
 		// Notice that we don't care whether we're up to date. User can perform textual
 		// merge via P4V:
 		const bool bIsCheckedOutElseWhere = State == EPerforceState::CheckedOutOther;
-		const bool bIsOpenForBranch = State == EPerforceState::Branched;
 		bCanDoCheckout =	bIsInP4NotCheckedOut ||
 							bIsCheckedOutElseWhere ||
 							bIsOpenForBranch;
@@ -283,7 +283,7 @@ bool FPerforceSourceControlState::CanCheckout() const
 	{
 		// For assets that are either binary or textual but marked for exclusive checkout
 		// we only want to permit check out when we are at head:
-		bCanDoCheckout = bIsInP4NotCheckedOut && IsCurrent();
+		bCanDoCheckout = (bIsInP4NotCheckedOut || bIsOpenForBranch) && IsCurrent();
 	}
 
 	return bCanDoCheckout;
diff --git a/Engine/Source/Developer/SourceControl/Private/SourceControlHelpers.cpp b/Engine/Source/Developer/SourceControl/Private/SourceControlHelpers.cpp
--- a/Engine/Source/Developer/SourceControl/Private/SourceControlHelpers.cpp
+++ b/Engine/Source/Developer/SourceControl/Private/SourceControlHelpers.cpp
@@ -1996,9 +1996,12 @@ namespace
 			FSourceControlStatePtr SourceControlState = SourceControlProvider.GetState(SourceFilename, StateCacheUsage);
 			if (SourceControlState.IsValid() && SourceControlState->IsSourceControlled())
 			{
+				FString DestinationPath = USourceControlHelpers::PackageFilename(DestPackage);
+				FSourceControlStatePtr DestinationState = SourceControlProvider.GetState(DestinationPath, StateCacheUsage);
 				TSharedRef<FCopy, ESPMode::ThreadSafe> CopyOperation = ISourceControlOperation::Create<FCopy>();
-				CopyOperation->SetDestination(USourceControlHelpers::PackageFilename(DestPackage));
+				CopyOperation->SetDestination(DestinationPath);
 				CopyOperation->CopyMethod = CopyMethod;
+				CopyOperation->ResolveMethod = (DestinationState.IsValid() && !DestinationState->CanRevert() && !IFileManager::Get().FileExists(*DestinationPath)) ? FCopy::EResolveMethod::AcceptSource : FCopy::EResolveMethod::AcceptTarget;
 
 				return (SourceControlProvider.Execute(CopyOperation, SourceFilename) == ECommandResult::Succeeded);
 			}
diff --git a/Engine/Source/Developer/SourceControl/Public/SourceControlOperations.h b/Engine/Source/Developer/SourceControl/Public/SourceControlOperations.h
--- a/Engine/Source/Developer/SourceControl/Public/SourceControlOperations.h
+++ b/Engine/Source/Developer/SourceControl/Public/SourceControlOperations.h
@@ -663,8 +663,17 @@ public:
 		Add      // The new file has no relation to the original file
 	};
 
+	enum class EResolveMethod
+	{
+		AcceptTarget,	// Resolve by accepting the destination file from the copy operation
+		AcceptSource,	// Resolve by accepting the source file from the copy operation
+	};
+
 	/** Whether a relationship to the original file should be maintained */
-	ECopyMethod CopyMethod = ECopyMethod::Branch; 
+	ECopyMethod CopyMethod = ECopyMethod::Branch;
+
+	/** The current state of the destination file */
+	EResolveMethod ResolveMethod = EResolveMethod::AcceptTarget;
 
 protected:
 	/** Destination path of the copy operation */
