API to get the list of revision control workspaces for the current host

#rnx

diff --git a/Engine/Plugins/Developer/PerforceSourceControl/Source/PerforceSourceControl/Private/PerforceConnection.cpp b/Engine/Plugins/Developer/PerforceSourceControl/Source/PerforceSourceControl/Private/PerforceConnection.cpp
--- a/Engine/Plugins/Developer/PerforceSourceControl/Source/PerforceSourceControl/Private/PerforceConnection.cpp
+++ b/Engine/Plugins/Developer/PerforceSourceControl/Source/PerforceSourceControl/Private/PerforceConnection.cpp
@@ -15,6 +15,7 @@
 #include "PerforceSourceControlProvider.h"
 #include "PerforceSourceControlSettings.h"
 #include "ProfilingDebugging/CpuProfilerTrace.h"
+#include "SourceControlOperations.h"
 
 #define LOCTEXT_NAMESPACE "PerforceConnection"
 
@@ -800,7 +801,7 @@ bool FPerforceConnection::AutoDetectWorkspace(const FPerforceConnectionInfo& InC
 	TArray<FString> ClientSpecList;
 	FSourceControlResultInfo ResultInfo;
 
-	Connection.GetWorkspaceList(InConnectionInfo, FOnIsCancelled(), ClientSpecList, ResultInfo);
+	Connection.GetProjectWorkspaces(InConnectionInfo, FOnIsCancelled(), ClientSpecList, ResultInfo);
 
 	OutErrorMessages = MoveTemp(ResultInfo.ErrorMessages);
 
@@ -1030,11 +1031,11 @@ bool FPerforceConnection::EnsureValidConnection(const FPerforceConnectionInfo& I
 	return bResult;
 }
 
-bool FPerforceConnection::GetWorkspaceList(const FPerforceConnectionInfo& InConnectionInfo, FOnIsCancelled InOnIsCanceled, TArray<FString>& OutWorkspaceList, FSourceControlResultInfo& OutResultInfo)
+bool FPerforceConnection::GetWorkspaces(const FPerforceConnectionInfo& InConnectionInfo, FOnIsCancelled InOnIsCancelled, TArray<FScWorkspaceInfo>& OutWorkspaceInfos, FSourceControlResultInfo& OutResultInfo)
 {
-	TRACE_CPUPROFILER_EVENT_SCOPE(FPerforceConnection::GetWorkspaceList)
+	TRACE_CPUPROFILER_EVENT_SCOPE(FPerforceConnection::GetWorkspaces)
 
-	if(bEstablishedConnection)
+	if (bEstablishedConnection)
 	{
 		TArray<FString> Params;
 		bool bAllowWildHosts = !GIsBuildMachine;
@@ -1043,14 +1044,12 @@ bool FPerforceConnection::GetWorkspaceList(const FPerforceConnectionInfo& InConn
 
 		FP4RecordSet Records;
 		bool bConnectionDropped = false;
-		bool bCommandOK = RunCommand(TEXT("clients"), Params, Records, OutResultInfo, InOnIsCanceled, bConnectionDropped);
+		bool bCommandOK = RunCommand(TEXT("clients"), Params, Records, OutResultInfo, InOnIsCancelled, bConnectionDropped);
 
 		if (bCommandOK)
 		{
-			FString ApplicationPath = IFileManager::Get().ConvertToAbsolutePathForExternalAppForRead(*ISourceControlModule::Get().GetSourceControlProjectDir()).ToLower();
-
 			FString LocalHostName = InConnectionInfo.HostOverride;
-			if(LocalHostName.Len() == 0)
+			if (LocalHostName.Len() == 0)
 			{
 				// No host override, check environment variable
 				LocalHostName = FPlatformMisc::GetEnvironmentVariable(TEXT("P4HOST"));
@@ -1070,39 +1069,17 @@ bool FPerforceConnection::GetWorkspaceList(const FPerforceConnectionInfo& InConn
 			{
 				const FString& ClientName = ClientRecord(TEXT("client"));
 				const FString& HostName = ClientRecord(TEXT("Host"));
-				FString ClientRootPath = ClientRecord(TEXT("Root")).ToLower();
+				const FString& RootDir = ClientRecord(TEXT("Root"));
 
 				//this clientspec has to be meant for this machine ( "" hostnames mean any host can use ths clientspec in p4 land)
 				bool bHostNameMatches = (LocalHostName == HostName.ToLower());
 				bool bHostNameWild = (HostName.Len() == 0);
 
-				if( bHostNameMatches || (bHostNameWild && bAllowWildHosts) )
+				if (bHostNameMatches || (bHostNameWild && bAllowWildHosts))
 				{
-					// A workspace root could be "null" which allows the user to map depot locations to different drives.
-					// Allow these workspaces since we already allow workspaces mapped to drive letters.
-					const bool bIsNullClientRootPath = (ClientRootPath == TEXT("null"));
-
-					//make sure all slashes point the same way
-					ClientRootPath = ClientRootPath.Replace(TEXT("\\"), TEXT("/"));
-					ApplicationPath = ApplicationPath.Replace(TEXT("\\"), TEXT("/"));
-
-					if (!ClientRootPath.EndsWith(TEXT("/")))
-					{
-						ClientRootPath += TEXT("/");
-					}
-
-					// Only allow paths that are ACTUALLY legit for this application
-					if (bIsNullClientRootPath || FPathViews::IsParentPathOf(ClientRootPath, ApplicationPath))
-					{
-						OutWorkspaceList.Add(ClientName);
-					}
-					else
-					{
-						UE_LOG(LogSourceControl, Verbose, TEXT(" %s client specs rejected due to root directory mismatch (%s)"), *ClientName, *ClientRootPath);
-					}
-
-					//Other useful fields: Description, Owner, Host
-
+					FScWorkspaceInfo& WorkspaceInfo = OutWorkspaceInfos.Emplace_GetRef();
+					WorkspaceInfo.Name = ClientName;
+					WorkspaceInfo.RootDir = RootDir;
 				}
 				else
 				{
@@ -1117,6 +1094,43 @@ bool FPerforceConnection::GetWorkspaceList(const FPerforceConnectionInfo& InConn
 	return false;
 }
 
+bool FPerforceConnection::GetProjectWorkspaces(const FPerforceConnectionInfo& InConnectionInfo, FOnIsCancelled InOnIsCancelled, TArray<FString>& OutWorkspaceNames, FSourceControlResultInfo& OutResultInfo)
+{
+	TArray<FScWorkspaceInfo> WorkspaceInfos;
+	if (GetWorkspaces(InConnectionInfo, MoveTemp(InOnIsCancelled), WorkspaceInfos, OutResultInfo))
+	{
+		FString ProjectDir = IFileManager::Get().ConvertToAbsolutePathForExternalAppForRead(*ISourceControlModule::Get().GetSourceControlProjectDir()).ToLower();
+		ProjectDir = ProjectDir.Replace(TEXT("\\"), TEXT("/"), ESearchCase::CaseSensitive);
+
+		for (FScWorkspaceInfo& WorkspaceInfo : WorkspaceInfos)
+		{
+			WorkspaceInfo.RootDir.ReplaceInline(TEXT("\\"), TEXT("/"), ESearchCase::CaseSensitive);
+
+			// A workspace root could be "null" which allows the user to map depot locations to different drives.
+			// Allow these workspaces since we already allow workspaces mapped to drive letters.
+			const bool bIsNullClientRootPath = (WorkspaceInfo.RootDir == TEXT("null"));
+
+			if (!WorkspaceInfo.RootDir.EndsWith(TEXT("/")))
+			{
+				WorkspaceInfo.RootDir += TEXT("/");
+			}
+
+			// Only allow paths that are ACTUALLY legit for this application
+			if (bIsNullClientRootPath || FPathViews::IsParentPathOf(WorkspaceInfo.RootDir, ProjectDir))
+			{
+				OutWorkspaceNames.Emplace(MoveTemp(WorkspaceInfo.Name));
+			}
+			else
+			{
+				UE_LOG(LogSourceControl, Verbose, TEXT(" %s client specs rejected due to root directory mismatch (%s)"), *WorkspaceInfo.Name, *WorkspaceInfo.RootDir);
+			}
+		}
+
+		return true;
+	}
+	return false;
+}
+
 bool FPerforceConnection::IsValidConnection()
 {
 	return bEstablishedConnection && !P4Client.Dropped();
diff --git a/Engine/Plugins/Developer/PerforceSourceControl/Source/PerforceSourceControl/Private/PerforceConnection.h b/Engine/Plugins/Developer/PerforceSourceControl/Source/PerforceSourceControl/Private/PerforceConnection.h
--- a/Engine/Plugins/Developer/PerforceSourceControl/Source/PerforceSourceControl/Private/PerforceConnection.h
+++ b/Engine/Plugins/Developer/PerforceSourceControl/Source/PerforceSourceControl/Private/PerforceConnection.h
@@ -8,6 +8,7 @@
 #include "Memory/SharedBuffer.h"
 
 class FPerforceSourceControlProvider;
+struct FScWorkspaceInfo;
 
 /** A map containing result of running Perforce command */
 class FP4Record : public TMap<FString, FString>
@@ -94,14 +95,24 @@ public:
 										FPerforceConnectionInfo& OutSettings, ISourceControlProvider::FInitResult::FConnectionErrors& OutConnectionErrors);
 
 	/**
-	 * Get List of ClientSpecs
+	 * Get all Perforce workspaces for the current host
 	 * @param InConnectionInfo	Connection credentials
 	 * @param InOnIsCancelled	Delegate called to check for operation cancellation.
-	 * @param OutWorkspaceList	The workspace list output.
+	 * @param OutWorkspaceInfos	Workspace infos output.
 	 * @param OutErrorMessages	Any error messages output.
 	 * @return - True if successful
 	 */
-	bool GetWorkspaceList(const FPerforceConnectionInfo& InConnectionInfo, FOnIsCancelled InOnIsCancelled, TArray<FString>& OutWorkspaceList, FSourceControlResultInfo& OutResultInfo);
+	bool GetWorkspaces(const FPerforceConnectionInfo& InConnectionInfo, FOnIsCancelled InOnIsCancelled, TArray<FScWorkspaceInfo>& OutWorkspaceInfos, FSourceControlResultInfo& OutResultInfo);
+
+	/**
+	 * Get Perforce workspaces under which of the current project is located
+	 * @param InConnectionInfo	Connection credentials
+	 * @param InOnIsCancelled	Delegate called to check for operation cancellation.
+	 * @param OutWorkspaceNames	Workspace names output.
+	 * @param OutErrorMessages	Any error messages output.
+	 * @return - True if successful
+	 */
+	bool GetProjectWorkspaces(const FPerforceConnectionInfo& InConnectionInfo, FOnIsCancelled InOnIsCancelled, TArray<FString>& OutWorkspaceNames, FSourceControlResultInfo& OutResultInfo);
 
 	/** Returns true if connection is currently active */
 	bool IsValidConnection();
diff --git a/Engine/Plugins/Developer/PerforceSourceControl/Source/PerforceSourceControl/Private/PerforceSourceControlInternalOperations.h b/Engine/Plugins/Developer/PerforceSourceControl/Source/PerforceSourceControl/Private/PerforceSourceControlInternalOperations.h
--- a/Engine/Plugins/Developer/PerforceSourceControl/Source/PerforceSourceControl/Private/PerforceSourceControlInternalOperations.h
+++ b/Engine/Plugins/Developer/PerforceSourceControl/Source/PerforceSourceControl/Private/PerforceSourceControlInternalOperations.h
@@ -12,15 +12,16 @@
 // module itself.
 
 /**
- * Internal-only source control operation for retrieving available workspaces.
+ * Internal-only source control operation for retrieving available workspaces for the current project
  */
-class FGetWorkspaces : public FSourceControlOperationBase
+
+class FGetProjectWorkspaces : public FSourceControlOperationBase
 {
 public:
 	// ISourceControlOperation interface
 	virtual FName GetName() const override
 	{
-		return "GetWorkspaces";
+		return "GetProjectWorkspaces";
 	}
 
 public:
diff --git a/Engine/Plugins/Developer/PerforceSourceControl/Source/PerforceSourceControl/Private/PerforceSourceControlOperations.cpp b/Engine/Plugins/Developer/PerforceSourceControl/Source/PerforceSourceControl/Private/PerforceSourceControlOperations.cpp
--- a/Engine/Plugins/Developer/PerforceSourceControl/Source/PerforceSourceControl/Private/PerforceSourceControlOperations.cpp
+++ b/Engine/Plugins/Developer/PerforceSourceControl/Source/PerforceSourceControl/Private/PerforceSourceControlOperations.cpp
@@ -64,6 +64,7 @@ void IPerforceSourceControlWorker::RegisterWorkers()
 	WorkersMap.Add("Sync", FGetPerforceSourceControlWorker::CreateStatic(&InstantiateWorker<FPerforceSyncWorker>));
 	WorkersMap.Add("CheckIn", FGetPerforceSourceControlWorker::CreateStatic(&InstantiateWorker<FPerforceCheckInWorker>));
 	WorkersMap.Add("GetWorkspaces", FGetPerforceSourceControlWorker::CreateStatic(&InstantiateWorker<FPerforceGetWorkspacesWorker>));
+	WorkersMap.Add("GetProjectWorkspaces", FGetPerforceSourceControlWorker::CreateStatic(&InstantiateWorker<FPerforceGetProjectWorkspacesWorker>));
 	WorkersMap.Add("Copy", FGetPerforceSourceControlWorker::CreateStatic(&InstantiateWorker<FPerforceCopyWorker>));
 	WorkersMap.Add("Resolve", FGetPerforceSourceControlWorker::CreateStatic(&InstantiateWorker<FPerforceResolveWorker>));
 	WorkersMap.Add("ChangeStatus", FGetPerforceSourceControlWorker::CreateStatic(&InstantiateWorker<FPerforceChangeStatusWorker>));
@@ -2228,12 +2229,12 @@ bool FPerforceGetWorkspacesWorker::Execute(FPerforceSourceControlCommand& InComm
 	if (!InCommand.IsCanceled() && ScopedConnection.IsValid())
 	{
 		FPerforceConnection& Connection = ScopedConnection.GetConnection();
-		TArray<FString> ClientSpecList;
-		InCommand.bCommandSuccessful = Connection.GetWorkspaceList(InCommand.ConnectionInfo, FOnIsCancelled::CreateRaw(&InCommand, &FPerforceSourceControlCommand::IsCanceled), ClientSpecList, InCommand.ResultInfo);
+		TArray<FScWorkspaceInfo> OutWorkspaces;
+		InCommand.bCommandSuccessful = Connection.GetWorkspaces(InCommand.ConnectionInfo, FOnIsCancelled::CreateRaw(&InCommand, &FPerforceSourceControlCommand::IsCanceled), OutWorkspaces, InCommand.ResultInfo);
 
 		check(InCommand.Operation->GetName() == GetName());
 		TSharedRef<FGetWorkspaces, ESPMode::ThreadSafe> Operation = StaticCastSharedRef<FGetWorkspaces>(InCommand.Operation);
-		Operation->Results = ClientSpecList;
+		Operation->Results = MoveTemp(OutWorkspaces);
 	}
 	return InCommand.bCommandSuccessful;
 }
@@ -2243,6 +2244,33 @@ bool FPerforceGetWorkspacesWorker::UpdateStates() const
 	return false;
 }
 
+FName FPerforceGetProjectWorkspacesWorker::GetName() const
+{
+	return "GetProjectWorkspaces";
+}
+
+bool FPerforceGetProjectWorkspacesWorker::Execute(FPerforceSourceControlCommand& InCommand)
+{
+	TRACE_CPUPROFILER_EVENT_SCOPE(FPerforceGetProjectWorkspacesWorker::Execute);
+
+	FScopedPerforceConnection ScopedConnection(InCommand);
+	if (!InCommand.IsCanceled() && ScopedConnection.IsValid())
+	{
+		FPerforceConnection& Connection = ScopedConnection.GetConnection();
+		TArray<FString> ClientSpecList;
+		InCommand.bCommandSuccessful = Connection.GetProjectWorkspaces(InCommand.ConnectionInfo, FOnIsCancelled::CreateRaw(&InCommand, &FPerforceSourceControlCommand::IsCanceled), ClientSpecList, InCommand.ResultInfo);
+
+		check(InCommand.Operation->GetName() == GetName());
+		TSharedRef<FGetProjectWorkspaces, ESPMode::ThreadSafe> Operation = StaticCastSharedRef<FGetProjectWorkspaces>(InCommand.Operation);
+		Operation->Results = ClientSpecList;
+	}
+	return InCommand.bCommandSuccessful;
+}
+
+bool FPerforceGetProjectWorkspacesWorker::UpdateStates() const
+{
+	return false;
+}
 
 FName FPerforceGetPendingChangelistsWorker::GetName() const
 {
diff --git a/Engine/Plugins/Developer/PerforceSourceControl/Source/PerforceSourceControl/Private/PerforceSourceControlOperations.h b/Engine/Plugins/Developer/PerforceSourceControl/Source/PerforceSourceControl/Private/PerforceSourceControlOperations.h
--- a/Engine/Plugins/Developer/PerforceSourceControl/Source/PerforceSourceControl/Private/PerforceSourceControlOperations.h
+++ b/Engine/Plugins/Developer/PerforceSourceControl/Source/PerforceSourceControl/Private/PerforceSourceControlOperations.h
@@ -213,6 +213,20 @@ public:
 	virtual bool UpdateStates() const override;
 };
 
+class FPerforceGetProjectWorkspacesWorker final : public IPerforceSourceControlWorker
+{
+public:
+	FPerforceGetProjectWorkspacesWorker(FPerforceSourceControlProvider& InSourceControlProvider)
+		: IPerforceSourceControlWorker(InSourceControlProvider)
+	{}
+	virtual ~FPerforceGetProjectWorkspacesWorker() = default;
+
+	// IPerforceSourceControlWorker interface
+	virtual FName GetName() const override;
+	virtual bool Execute(FPerforceSourceControlCommand& InCommand) override;
+	virtual bool UpdateStates() const override;
+};
+
 class FPerforceGetPendingChangelistsWorker final : public IPerforceSourceControlWorker
 {
 public:
diff --git a/Engine/Plugins/Developer/PerforceSourceControl/Source/PerforceSourceControl/Private/SPerforceSourceControlSettings.cpp b/Engine/Plugins/Developer/PerforceSourceControl/Source/PerforceSourceControl/Private/SPerforceSourceControlSettings.cpp
--- a/Engine/Plugins/Developer/PerforceSourceControl/Source/PerforceSourceControl/Private/SPerforceSourceControlSettings.cpp
+++ b/Engine/Plugins/Developer/PerforceSourceControl/Source/PerforceSourceControl/Private/SPerforceSourceControlSettings.cpp
@@ -305,7 +305,7 @@ void SPerforceSourceControlSettings::QueryWorkspaces()
 		CurrentWorkspace = FString();
 
 		// fire off the workspace query
-		GetWorkspacesOperation = ISourceControlOperation::Create<FGetWorkspaces>();
+		GetWorkspacesOperation = ISourceControlOperation::Create<FGetProjectWorkspaces>();
 		GetSCCProvider().Execute(GetWorkspacesOperation.ToSharedRef(), EConcurrency::Asynchronous, FSourceControlOperationComplete::CreateSP(this, &SPerforceSourceControlSettings::OnSourceControlOperationComplete) );
 
 		State = ESourceControlOperationState::Querying;
@@ -316,8 +316,8 @@ void SPerforceSourceControlSettings::OnSourceControlOperationComplete(const FSou
 {
 	if(InResult == ECommandResult::Succeeded)
 	{
-		check(InOperation->GetName() == "GetWorkspaces");
-		check(GetWorkspacesOperation == StaticCastSharedRef<FGetWorkspaces>(InOperation));
+		check(InOperation->GetName() == "GetProjectWorkspaces");
+		check(GetWorkspacesOperation == StaticCastSharedRef<FGetProjectWorkspaces>(InOperation));
 
 		// refresh workspaces list from operation results
 		Workspaces.Empty();
diff --git a/Engine/Plugins/Developer/PerforceSourceControl/Source/PerforceSourceControl/Private/SPerforceSourceControlSettings.h b/Engine/Plugins/Developer/PerforceSourceControl/Source/PerforceSourceControl/Private/SPerforceSourceControlSettings.h
--- a/Engine/Plugins/Developer/PerforceSourceControl/Source/PerforceSourceControl/Private/SPerforceSourceControlSettings.h
+++ b/Engine/Plugins/Developer/PerforceSourceControl/Source/PerforceSourceControl/Private/SPerforceSourceControlSettings.h
@@ -16,7 +16,7 @@
 #include "SourceControlOperationBase.h"
 #include "ISourceControlProvider.h"
 
-class FGetWorkspaces;
+class FGetProjectWorkspaces;
 class FPerforceSourceControlProvider;
 
 class SButton;
@@ -141,7 +141,7 @@ private:
 	TSharedPtr<SComboButton> WorkspaceCombo;
 
 	/** The source control operation in progress */
-	TSharedPtr<FGetWorkspaces, ESPMode::ThreadSafe> GetWorkspacesOperation;
+	TSharedPtr<FGetProjectWorkspaces, ESPMode::ThreadSafe> GetWorkspacesOperation;
 
 	/** Pointer to the password text box widget */
 	static TWeakPtr<SEditableTextBox> PasswordTextBox;
diff --git a/Engine/Source/Developer/SourceControl/Public/SourceControlOperations.h b/Engine/Source/Developer/SourceControl/Public/SourceControlOperations.h
--- a/Engine/Source/Developer/SourceControl/Public/SourceControlOperations.h
+++ b/Engine/Source/Developer/SourceControl/Public/SourceControlOperations.h
@@ -1472,6 +1472,30 @@ protected:
 	FString WorkspaceName;
 };
 
+/** Information on a source control workspace */
+struct FScWorkspaceInfo
+{
+	FString Name;
+	FString RootDir;
+};
+
+/**
+ * Operation used to get all workspaces for the current host if the source control system supports this functionality.
+ */
+class FGetWorkspaces : public FSourceControlOperationBase
+{
+public:
+	// ISourceControlOperation interface
+	virtual FName GetName() const override
+	{
+		return "GetWorkspaces";
+	}
+
+public:
+	/** Output array of workspace info */
+	TArray<FScWorkspaceInfo> Results;
+};
+
 /**
  * This operation uses p4v print command to get file from specified depot by shelved changelist or file revision returns a package filename
  */
