Unshelved from pending changelist '36726159':

PR #12389: [SourceControl] Add possibility to monitor an array of roots in SourceControlProvider state branches, for now only implemented for p4

diff --git a/Engine/Plugins/Developer/PerforceSourceControl/Source/PerforceSourceControl/Private/PerforceSourceControlCommand.h b/Engine/Plugins/Developer/PerforceSourceControl/Source/PerforceSourceControl/Private/PerforceSourceControlCommand.h
--- a/Engine/Plugins/Developer/PerforceSourceControl/Source/PerforceSourceControl/Private/PerforceSourceControlCommand.h
+++ b/Engine/Plugins/Developer/PerforceSourceControl/Source/PerforceSourceControl/Private/PerforceSourceControlCommand.h
@@ -109,7 +109,7 @@ public:
 	/** Branch names for status queries */
 	TArray< FString > StatusBranchNames;
 
-	/** Content root for branch mapping */
-	FString ContentRoot;
+	/** Content roots for branch mapping */
+	TArray<FString> ContentRoots;
 
 };
diff --git a/Engine/Plugins/Developer/PerforceSourceControl/Source/PerforceSourceControl/Private/PerforceSourceControlOperations.cpp b/Engine/Plugins/Developer/PerforceSourceControl/Source/PerforceSourceControl/Private/PerforceSourceControlOperations.cpp
--- a/Engine/Plugins/Developer/PerforceSourceControl/Source/PerforceSourceControl/Private/PerforceSourceControlOperations.cpp
+++ b/Engine/Plugins/Developer/PerforceSourceControl/Source/PerforceSourceControl/Private/PerforceSourceControlOperations.cpp
@@ -1142,7 +1142,7 @@ bool FPerforceSyncWorker::UpdateStates() const
 	return UpdateCachedStates(GetSCCProvider(), OutResults);
 }
 
-static void ParseBranchModificationResults(const FP4RecordSet& InRecords, const TArray<FText>& ErrorMessages, const FString& ContentRoot, TMap<FString, FBranchModification>& BranchModifications)
+static void ParseBranchModificationResults(const FP4RecordSet& InRecords, const TArray<FText>& ErrorMessages, const TArray<FString>& ContentRoots, TMap<FString, FBranchModification>& BranchModifications)
 {
 	TRACE_CPUPROFILER_EVENT_SCOPE(FPerforce::ParseBranchModificationResults);
 
@@ -1165,11 +1165,16 @@ static void ParseBranchModificationResults(const FP4RecordSet& InRecords, const
 		// Get the content filename and add to branch states
 		FString CurrentBranch(TEXT("*CurrentBranch"));
 		FString Branch, BranchFile;
-		if (DepotFileName.Split(ContentRoot, &Branch, &BranchFile))
+		for (const FString& ContentRoot : ContentRoots)
 		{
-			// Sanitize names
-			Branch.RemoveFromEnd(FString(TEXT("/")));
-			BranchFile.RemoveFromStart(FString(TEXT("/")));
+			if (DepotFileName.Split(ContentRoot, &Branch, &BranchFile))
+			{
+				// Sanitize names
+				Branch.RemoveFromEnd(FString(TEXT("/")));
+				BranchFile.RemoveFromStart(FString(TEXT("/")));
+
+				break;
+			}
 		}
 
 		if (!Branch.Len() || !BranchFile.Len())
@@ -1229,7 +1234,7 @@ static void ParseBranchModificationResults(const FP4RecordSet& InRecords, const
 	}
 }
 
-static void ParseUpdateStatusResults(const FP4RecordSet& InRecords, const TArray<FText>& ErrorMessages, TArray<FPerforceSourceControlState>& OutStates, const FString& ContentRoot, TMap<FString, FBranchModification>& BranchModifications)
+static void ParseUpdateStatusResults(const FP4RecordSet& InRecords, const TArray<FText>& ErrorMessages, TArray<FPerforceSourceControlState>& OutStates, const TArray<FString>& ContentRoots, TMap<FString, FBranchModification>& BranchModifications)
 {
 	TRACE_CPUPROFILER_EVENT_SCOPE(FPerforceUpdateStatusWorker::ParseUpdateStatusResults);
 
@@ -1249,16 +1254,21 @@ static void ParseUpdateStatusResults(const FP4RecordSet& InRecords, const TArray
 
 		FString Branch;
 		FString FileName(ClientFile);
-		if (DepotFileName.Split(ContentRoot, &Branch, &FileName))
+		for (const FString& ContentRoot : ContentRoots)
 		{
-			// Sanitize
-			Branch.RemoveFromEnd(FString(TEXT("/")));
-			FileName.RemoveFromStart(FString(TEXT("/")));
-
-			// Add to branch modifications if not currently recorded
-			if (FileName.Len() && !BranchModifications.Contains(FileName))
+			if (DepotFileName.Split(ContentRoot, &Branch, &FileName))
 			{
-				BranchModifications.Add(FileName, FBranchModification(Branch, FileName, FString(TEXT("none")), 0, 0));
+				// Sanitize
+				Branch.RemoveFromEnd(FString(TEXT("/")));
+				FileName.RemoveFromStart(FString(TEXT("/")));
+
+				// Add to branch modifications if not currently recorded
+				if (FileName.Len() && !BranchModifications.Contains(FileName))
+				{
+					BranchModifications.Add(FileName, FBranchModification(Branch, FileName, FString(TEXT("none")), 0, 0));
+				}
+
+				break;
 			}
 		}
 
@@ -1325,11 +1335,16 @@ static void ParseUpdateStatusResults(const FP4RecordSet& InRecords, const TArray
 		FString Branch;
 		FString BranchFile;
 
-		if (DepotFileName.Split(ContentRoot, &Branch, &BranchFile))
+		for (const FString& ContentRoot : ContentRoots)
 		{
-			// Sanitize
-			Branch.RemoveFromEnd(FString(TEXT("/")));
-			BranchFile.RemoveFromStart(FString(TEXT("/")));
+			if (DepotFileName.Split(ContentRoot, &Branch, &BranchFile))
+			{
+				// Sanitize
+				Branch.RemoveFromEnd(FString(TEXT("/")));
+				BranchFile.RemoveFromStart(FString(TEXT("/")));
+
+				break;
+			}
 		}
 
 		State.State = EPerforceState::ReadOnly;
@@ -1912,7 +1927,7 @@ bool FPerforceUpdateStatusWorker::Execute(FPerforceSourceControlCommand& InComma
 			Parameters.Add(TEXT("-Or"));
 
 			// Get the branches of interest for status updates
-			const FString& ContentRoot = InCommand.ContentRoot;
+			const TArray<FString>& ContentRoots = InCommand.ContentRoots;
 			const TArray<FString>& StatusBranches = InCommand.StatusBranchNames;
 
 			// Mandatory parameters (the list of files to stat):
@@ -1927,19 +1942,22 @@ bool FPerforceUpdateStatusWorker::Execute(FPerforceSourceControlCommand& InComma
 				{
 					for (const FString& Branch : StatusBranches)
 					{
-						// Check the status branch for updates
-						FString BranchFile;
-						if (File.Split(ContentRoot, nullptr, &BranchFile))
+						for (const FString& ContentRoot : ContentRoots)
 						{
-							// Ignore collection files when querying status branches
-							FString Ext = FPaths::GetExtension(BranchFile, true);
-							if (Ext.Compare(TEXT(".collection"), ESearchCase::IgnoreCase) == 0)
+							// Check the status branch for updates
+							FString BranchFile;
+							if (File.Split(ContentRoot, nullptr, &BranchFile))
 							{
-								continue;
+								// Ignore collection files when querying status branches
+								FString Ext = FPaths::GetExtension(BranchFile, true);
+								if (Ext.Compare(TEXT(".collection"), ESearchCase::IgnoreCase) == 0)
+								{
+									continue;
+								}
+
+								TArray<FStringFormatArg> Args = { Branch, ContentRoot, BranchFile };
+								Parameters.Add(FString::Format(TEXT("{0}/{1}{2}"), Args));
 							}
-							
-							TArray<FStringFormatArg> Args = { Branch, ContentRoot, BranchFile };
-							Parameters.Add(FString::Format(TEXT("{0}/{1}{2}"), Args));
 						}
 					}
 				}
@@ -1973,12 +1991,12 @@ bool FPerforceUpdateStatusWorker::Execute(FPerforceSourceControlCommand& InComma
 
 				FP4RecordSet RevisionRecords;
 				InCommand.bCommandSuccessful &= Connection.RunCommand(TEXT("fstat"), RevisionParameters, RevisionRecords, nullptr, InCommand.ResultInfo, FOnIsCancelled::CreateRaw(&InCommand, &FPerforceSourceControlCommand::IsCanceled), InCommand.bConnectionDropped, Flags);
-				ParseBranchModificationResults(RevisionRecords, InCommand.ResultInfo.ErrorMessages, ContentRoot, BranchModifications);
+				ParseBranchModificationResults(RevisionRecords, InCommand.ResultInfo.ErrorMessages, ContentRoots, BranchModifications);
 			}
 
 			FP4RecordSet Records;
 			InCommand.bCommandSuccessful &= Connection.RunCommand(TEXT("fstat"), Parameters, Records, nullptr, InCommand.ResultInfo, FOnIsCancelled::CreateRaw(&InCommand, &FPerforceSourceControlCommand::IsCanceled), InCommand.bConnectionDropped, Flags);
-			ParseUpdateStatusResults(Records, InCommand.ResultInfo.ErrorMessages, OutStates, ContentRoot, BranchModifications);
+			ParseUpdateStatusResults(Records, InCommand.ResultInfo.ErrorMessages, OutStates, ContentRoots, BranchModifications);
 			RemoveRedundantErrors(InCommand, TEXT(" - no such file(s)."), false);
 			RemoveRedundantErrors(InCommand, TEXT("' is not under client's root '"));
 			RemoveRedundantErrors(InCommand, TEXT(" - file(s) not in client view"), false);
diff --git a/Engine/Plugins/Developer/PerforceSourceControl/Source/PerforceSourceControl/Private/PerforceSourceControlProvider.cpp b/Engine/Plugins/Developer/PerforceSourceControl/Source/PerforceSourceControl/Private/PerforceSourceControlProvider.cpp
--- a/Engine/Plugins/Developer/PerforceSourceControl/Source/PerforceSourceControl/Private/PerforceSourceControlProvider.cpp
+++ b/Engine/Plugins/Developer/PerforceSourceControl/Source/PerforceSourceControl/Private/PerforceSourceControlProvider.cpp
@@ -551,7 +551,7 @@ ECommandResult::Type FPerforceSourceControlProvider::Execute( const FSourceContr
 		Command->bAutoDelete = false;
 		Command->Files = AbsoluteFiles;
 		Command->StatusBranchNames = StatusBranchNames;
-		Command->ContentRoot = ContentRoot;
+		Command->ContentRoots = ContentRoots;
 		Command->OperationCompleteDelegate = InOperationCompleteDelegate;
 		Command->Changelist = Changelist;
 		return ExecuteSynchronousCommand(*Command, InOperation->GetInProgressString(), true);
@@ -562,7 +562,7 @@ ECommandResult::Type FPerforceSourceControlProvider::Execute( const FSourceContr
 		Command->bAutoDelete = true;
 		Command->Files = AbsoluteFiles;
 		Command->StatusBranchNames = StatusBranchNames;
-		Command->ContentRoot = ContentRoot;
+		Command->ContentRoots = ContentRoots;
 		Command->OperationCompleteDelegate = InOperationCompleteDelegate;
 		Command->Changelist = Changelist;
 		return IssueCommand(*Command, false);
@@ -881,7 +881,7 @@ bool FPerforceSourceControlProvider::TryToDownloadFileFromBackgroundThread(const
 	Command.bAutoDelete = false;
 	Command.Files = SourceControlHelpers::AbsoluteFilenames(InFiles);
 	Command.StatusBranchNames = StatusBranchNames;
-	Command.ContentRoot = ContentRoot;
+	Command.ContentRoots = ContentRoots;
 
 	// DoWork will use a shared connection, we need to call DoThreadedWork to make sure that 
 	// we use our own connection for this background thread.
@@ -1153,9 +1153,14 @@ bool FPerforceSourceControlProvider::QueryStateBranchConfig(const FString& Confi
 }
 
 void FPerforceSourceControlProvider::RegisterStateBranches(const TArray<FString>& BranchNames, const FString& ContentRootIn)
+{
+	return RegisterStateBranches(BranchNames, TArray<FString> { ContentRootIn });
+}
+
+void FPerforceSourceControlProvider::RegisterStateBranches(const TArray<FString>& BranchNames, const TArray<FString>& ContentRootsIn)
 {
 	StatusBranchNames = BranchNames;
-	ContentRoot = ContentRootIn;
+	ContentRoots = ContentRootsIn;
 }
 
 int32 FPerforceSourceControlProvider::GetStateBranchIndex(const FString& BranchName) const
diff --git a/Engine/Plugins/Developer/PerforceSourceControl/Source/PerforceSourceControl/Private/PerforceSourceControlProvider.h b/Engine/Plugins/Developer/PerforceSourceControl/Source/PerforceSourceControl/Private/PerforceSourceControlProvider.h
--- a/Engine/Plugins/Developer/PerforceSourceControl/Source/PerforceSourceControl/Private/PerforceSourceControlProvider.h
+++ b/Engine/Plugins/Developer/PerforceSourceControl/Source/PerforceSourceControl/Private/PerforceSourceControlProvider.h
@@ -36,6 +36,7 @@ public:
 	virtual const FName& GetName(void) const override;
 	virtual bool QueryStateBranchConfig(const FString& ConfigSrc, const FString& ConfigDest) override;
 	virtual void RegisterStateBranches(const TArray<FString>& BranchNames, const FString& ContentRootIn) override;
+	virtual void RegisterStateBranches(const TArray<FString>& BranchNames, const TArray<FString>& ContentRootsIn) override;
 	virtual int32 GetStateBranchIndex(const FString& BranchName) const override;
 	virtual ECommandResult::Type GetState( const TArray<FString>& InFiles, TArray<FSourceControlStateRef>& OutState, EStateCacheUsage::Type InStateCacheUsage ) override;
 	virtual ECommandResult::Type GetState(const TArray<FSourceControlChangelistRef>& InChangelists, TArray<FSourceControlChangelistStateRef>& OutState, EStateCacheUsage::Type InStateCacheUsage) override;
@@ -211,7 +212,7 @@ private:
 	/** Array of branch names for status queries */
 	TArray<FString> StatusBranchNames;
 
-	/** Content root for branch status query mapping */
-	FString ContentRoot;
+	/** Content roots for branch status query mapping */
+	TArray<FString> ContentRoots;
 
 };
diff --git a/Engine/Source/Developer/SourceControl/Private/ISourceControlProvider.cpp b/Engine/Source/Developer/SourceControl/Private/ISourceControlProvider.cpp
--- a/Engine/Source/Developer/SourceControl/Private/ISourceControlProvider.cpp
+++ b/Engine/Source/Developer/SourceControl/Private/ISourceControlProvider.cpp
@@ -27,6 +27,15 @@ ECommandResult::Type ISourceControlProvider::Login(const FString& InPassword, EC
 	return Execute(ConnectOperation, InConcurrency, InOperationCompleteDelegate);
 }
 
+void ISourceControlProvider::RegisterStateBranches(const TArray<FString>& BranchNames, const TArray<FString>& ContentRoots)
+{
+	if (ContentRoots.Num() > 0)
+	{
+		ensureMsgf(ContentRoots.Num() >= 1, TEXT("This provider does not support multiple content roots, only using the first provided root (%s)"), *ContentRoots[0]);
+		RegisterStateBranches(BranchNames, ContentRoots[0]);
+	}
+}
+
 ECommandResult::Type ISourceControlProvider::GetState(const TArray<UPackage*>& InPackages, TArray<FSourceControlStateRef>& OutState, EStateCacheUsage::Type InStateCacheUsage)
 {
 	TArray<FString> Files = SourceControlHelpers::PackageFilenames(InPackages);
diff --git a/Engine/Source/Developer/SourceControl/Public/ISourceControlProvider.h b/Engine/Source/Developer/SourceControl/Public/ISourceControlProvider.h
--- a/Engine/Source/Developer/SourceControl/Public/ISourceControlProvider.h
+++ b/Engine/Source/Developer/SourceControl/Public/ISourceControlProvider.h
@@ -221,6 +221,13 @@ public:
 	*/
 	virtual void RegisterStateBranches(const TArray<FString>& BranchNames, const FString& ContentRoot) = 0;
 
+	/**
+	* Register branches to query for state in addition to the current branch
+	* @param	BranchNames			Names of the branches to query
+	* @param	ContentRoots		Paths to the content roots for branch mapping
+	*/
+	SOURCECONTROL_API virtual void RegisterStateBranches( const TArray<FString>& BranchNames, const TArray<FString>& ContentRoots );
+
 	/**
 	*Gets the state index of the specified branch, higher index branches are generally closer to releases
 	* @param	BranchName			Names of the branches to query
