Adding support for files marked for integration in the perforce source control, making them visible for the changelist shelve states

[RN] minor Creator - Source Control

diff --git a/Engine/Plugins/Developer/PerforceSourceControl/Source/PerforceSourceControl/Private/PerforceSourceControlOperations.cpp b/Engine/Plugins/Developer/PerforceSourceControl/Source/PerforceSourceControl/Private/PerforceSourceControlOperations.cpp
--- a/Engine/Plugins/Developer/PerforceSourceControl/Source/PerforceSourceControl/Private/PerforceSourceControlOperations.cpp
+++ b/Engine/Plugins/Developer/PerforceSourceControl/Source/PerforceSourceControl/Private/PerforceSourceControlOperations.cpp
@@ -317,7 +317,7 @@ static void ParseRecordSetForState(const FP4RecordSet& InRecords, TMap<FString,
 					OutResults.Add(FullPath, EPerforceState::ReadOnly);
 				}
 			}
-			else if(Action == TEXT("branch"))
+			else if(Action == TEXT("branch") || Action == TEXT("integrate"))
 			{
 				OutResults.Add(FullPath, EPerforceState::Branched);
 			}
@@ -339,16 +339,23 @@ static void UpdateChangelistState(FPerforceSourceControlProvider& SCCProvider, c
 
 		for (TMap<FString, EPerforceState::Type>::TConstIterator It(InOperationResults); It; ++It)
 		{
-			if ((It.Value() != EPerforceState::CheckedOut) && (It.Value() != EPerforceState::OpenForAdd) && It.Value() != EPerforceState::MarkedForDelete)
-			{
-				continue;
-			}
+			TSharedRef<FPerforceSourceControlState, ESPMode::ThreadSafe> State;
 
-			TSharedRef<FPerforceSourceControlState, ESPMode::ThreadSafe> State = SCCProvider.GetStateInternal(It.Key());
-			ChangelistState->Files.Add(State);
+			switch (It.Value())
+			{
+				case EPerforceState::CheckedOut:
+				case EPerforceState::OpenForAdd:
+				case EPerforceState::MarkedForDelete:
+				case EPerforceState::Branched:
+					State = SCCProvider.GetStateInternal(It.Key());
+					ChangelistState->Files.Add(State);
 
-			// Keep the changelist stored with cached file state in sync with the actual changelist that owns this file.
-			State->Changelist = InChangelist;
+					// Keep the changelist stored with cached file state in sync with the actual changelist that owns this file.
+					State->Changelist = InChangelist;
+					break;
+				default:					
+					break;
+			}
 		}
 	}
 }
@@ -1358,7 +1365,7 @@ static void ParseUpdateStatusResults(const FP4RecordSet& InRecords, const TArray
 		}
 		else if (OpenType.Len() > 0)
 		{
-			if(Action.Len() > 0 && Action == TEXT("branch"))
+			if(Action.Len() > 0 && (Action == TEXT("branch") || Action == TEXT("integrate")))
 			{
 				State.State = EPerforceState::Branched;
 			}
@@ -1625,6 +1632,10 @@ static void ParseOpenedResults(const FP4RecordSet& InRecords, const FString& Cli
 			{
 				OutState.State = EPerforceState::MarkedForDelete;
 			}
+			else if (Action == TEXT("branch") || Action == TEXT("integrate"))
+			{
+				OutState.State = EPerforceState::Branched;
+			}
 		}
 
 		const FString& Changelist = ClientRecord(TEXT("change"));
@@ -1680,6 +1691,10 @@ static void ParseShelvedResults(const FP4RecordSet& InRecords, TMap<FString, EPe
 			{
 				OutResults.Add(DepotFileName, EPerforceState::MarkedForDelete);
 			}
+			else if (Action == TEXT("branch") || Action == TEXT("integrate"))
+			{
+				OutResults.Add(DepotFileName, EPerforceState::Branched);
+			}
 		}
 	}
 }
@@ -1712,6 +1727,10 @@ static void ParseShelvedChangelistResults(const FP4RecordSet& InRecords, TMap<FS
 			{
 				OutResults.Add(DepotFileName, EPerforceState::MarkedForDelete);
 			}
+			else if (Action == TEXT("branch") || Action == TEXT("integrate"))
+			{
+				OutResults.Add(DepotFileName, EPerforceState::Branched);
+			}
 		}
 	}
 }
