Fixing incorrect usage of TSharedRef causing breakages

diff --git a/Engine/Plugins/Developer/PerforceSourceControl/Source/PerforceSourceControl/Private/PerforceSourceControlOperations.cpp b/Engine/Plugins/Developer/PerforceSourceControl/Source/PerforceSourceControl/Private/PerforceSourceControlOperations.cpp
--- a/Engine/Plugins/Developer/PerforceSourceControl/Source/PerforceSourceControl/Private/PerforceSourceControlOperations.cpp
+++ b/Engine/Plugins/Developer/PerforceSourceControl/Source/PerforceSourceControl/Private/PerforceSourceControlOperations.cpp
@@ -339,7 +339,7 @@ static void UpdateChangelistState(FPerforceSourceControlProvider& SCCProvider, c
 
 		for (TMap<FString, EPerforceState::Type>::TConstIterator It(InOperationResults); It; ++It)
 		{
-			TSharedRef<FPerforceSourceControlState, ESPMode::ThreadSafe> State;
+			TSharedPtr<FPerforceSourceControlState, ESPMode::ThreadSafe> State;
 
 			switch (It.Value())
 			{
@@ -348,7 +348,7 @@ static void UpdateChangelistState(FPerforceSourceControlProvider& SCCProvider, c
 				case EPerforceState::MarkedForDelete:
 				case EPerforceState::Branched:
 					State = SCCProvider.GetStateInternal(It.Key());
-					ChangelistState->Files.Add(State);
+					ChangelistState->Files.Add(State.ToSharedRef());
 
 					// Keep the changelist stored with cached file state in sync with the actual changelist that owns this file.
 					State->Changelist = InChangelist;
