Don't use plural on source control tool tips since that only applies to a single file

diff --git a/Engine/Plugins/Developer/PerforceSourceControl/Source/PerforceSourceControl/Private/PerforceSourceControlState.cpp b/Engine/Plugins/Developer/PerforceSourceControl/Source/PerforceSourceControl/Private/PerforceSourceControlState.cpp
--- a/Engine/Plugins/Developer/PerforceSourceControl/Source/PerforceSourceControl/Private/PerforceSourceControlState.cpp
+++ b/Engine/Plugins/Developer/PerforceSourceControl/Source/PerforceSourceControl/Private/PerforceSourceControlState.cpp
@@ -201,11 +201,11 @@ FText FPerforceSourceControlState::GetDisplayTooltip() const
 {
 	if (IsConflicted())
 	{
-		return LOCTEXT("Conflicted_Tooltip", "The file(s) have local changes that need to be resolved with changes submitted to the Perforce depot");
+		return LOCTEXT("Conflicted_Tooltip", "The file has local changes that need to be resolved with changes submitted to the Perforce depot");
 	}
 	else if( !IsCurrent() )
 	{
-		return LOCTEXT("NotCurrent_Tooltip", "The file(s) are not at the head revision");
+		return LOCTEXT("NotCurrent_Tooltip", "The file is not at the head revision");
 	}
 	else if (State != EPerforceState::CheckedOut && State != EPerforceState::CheckedOutOther)
 	{
@@ -225,23 +225,23 @@ FText FPerforceSourceControlState::GetDisplayTooltip() const
 	{
 	default:
 	case EPerforceState::DontCare:
-		return LOCTEXT("Unknown_Tooltip", "The file(s) status is unknown");
+		return LOCTEXT("Unknown_Tooltip", "The file status is unknown");
 	case EPerforceState::CheckedOut:
-		return LOCTEXT("CheckedOut_Tooltip", "The file(s) are checked out");
+		return LOCTEXT("CheckedOut_Tooltip", "The file is checked out");
 	case EPerforceState::ReadOnly:
-		return LOCTEXT("ReadOnly_Tooltip", "The file(s) are marked locally as read-only");
+		return LOCTEXT("ReadOnly_Tooltip", "The file is marked locally as read-only");
 	case EPerforceState::NotInDepot:
-		return LOCTEXT("NotInDepot_Tooltip", "The file(s) are not present in the Perforce depot");
+		return LOCTEXT("NotInDepot_Tooltip", "The file is not present in the Perforce depot");
 	case EPerforceState::CheckedOutOther:
 		return FText::Format( LOCTEXT("CheckedOutOther_Tooltip", "Checked out by: {0}"), FText::FromString(OtherUserCheckedOut) );
 	case EPerforceState::Ignore:
-		return LOCTEXT("Ignore_Tooltip", "The file(s) are ignored by Perforce");
+		return LOCTEXT("Ignore_Tooltip", "The file is ignored by Perforce");
 	case EPerforceState::OpenForAdd:
-		return LOCTEXT("OpenedForAdd_Tooltip", "The file(s) are opened for add");
+		return LOCTEXT("OpenedForAdd_Tooltip", "The file is opened for add");
 	case EPerforceState::MarkedForDelete:
-		return LOCTEXT("MarkedForDelete_Tooltip", "The file(s) are marked for delete");
+		return LOCTEXT("MarkedForDelete_Tooltip", "The file is marked for delete");
 	case EPerforceState::Branched:
-		return LOCTEXT("Branched_Tooltip", "The file(s) are opened for branching");
+		return LOCTEXT("Branched_Tooltip", "The file is opened for branching");
 	}
 }
 
