Add a cvar to allow reverting deleted files in perforce

diff --git a/Engine/Plugins/Developer/PerforceSourceControl/Source/PerforceSourceControl/Private/PerforceSourceControlState.cpp b/Engine/Plugins/Developer/PerforceSourceControl/Source/PerforceSourceControl/Private/PerforceSourceControlState.cpp
--- a/Engine/Plugins/Developer/PerforceSourceControl/Source/PerforceSourceControl/Private/PerforceSourceControlState.cpp
+++ b/Engine/Plugins/Developer/PerforceSourceControl/Source/PerforceSourceControl/Private/PerforceSourceControlState.cpp
@@ -19,6 +19,20 @@ FPerforceSourceControlState& FPerforceSourceControlState::operator=(const FPerfo
 FPerforceSourceControlState& FPerforceSourceControlState::operator=(FPerforceSourceControlState&& Other) noexcept = default;
 PRAGMA_ENABLE_DEPRECATION_WARNINGS
 
+namespace UE::RevisionControl::Perforce
+{
+	static bool bAllowRevertDeletedFiles = false;
+	
+	static FAutoConsoleVariableRef CVarAllowRevertDeletedFiles(
+		TEXT("RevisionControl.Perforce.AllowRevertingDeletedFiles"),
+		bAllowRevertDeletedFiles,
+		TEXT("If true, deleted files can be reverted in the editor"),
+		ECVF_Default
+	);
+}
+
+
+
 int32 FPerforceSourceControlState::GetHistorySize() const
 {
 	return History.Num();
@@ -355,7 +369,14 @@ bool FPerforceSourceControlState::CanAdd() const
 
 bool FPerforceSourceControlState::CanRevert() const
 {
-	return IsCheckedOut() || IsAdded();
+	bool bCanRevert = IsCheckedOut() || IsAdded();
+
+	if (UE::RevisionControl::Perforce::CVarAllowRevertDeletedFiles->GetBool())
+	{
+		bCanRevert |= IsDeleted();
+	}
+
+	return bCanRevert;
 }
 
 void FPerforceSourceControlState::Update(const FPerforceSourceControlState& InOther, const FDateTime* InTimeStamp /* = nullptr */)
