PerforceSourceControlState: Add missing include file causing test failure

diff --git a/Engine/Plugins/Developer/PerforceSourceControl/Source/PerforceSourceControl/Private/PerforceSourceControlState.cpp b/Engine/Plugins/Developer/PerforceSourceControl/Source/PerforceSourceControl/Private/PerforceSourceControlState.cpp
--- a/Engine/Plugins/Developer/PerforceSourceControl/Source/PerforceSourceControl/Private/PerforceSourceControlState.cpp
+++ b/Engine/Plugins/Developer/PerforceSourceControl/Source/PerforceSourceControl/Private/PerforceSourceControlState.cpp
@@ -3,6 +3,7 @@
 #include "PerforceSourceControlState.h"
 #include "PerforceSourceControlRevision.h"
 #include "Misc/EngineVersion.h"
+#include "HAL/IConsoleManager.h"
 
 #if SOURCE_CONTROL_WITH_SLATE
 #include "Textures/SlateIcon.h"
