[TEDS Revision Control] Change the Revision Control table name to Package Table since it reflects the data better.

#Synced-CL 51541229

diff --git a/Engine/Plugins/Developer/PerforceSourceControl/Source/PerforceSourceControl/Private/PerforceSourceControlOperations.cpp b/Engine/Plugins/Developer/PerforceSourceControl/Source/PerforceSourceControl/Private/PerforceSourceControlOperations.cpp
--- a/Engine/Plugins/Developer/PerforceSourceControl/Source/PerforceSourceControl/Private/PerforceSourceControlOperations.cpp
+++ b/Engine/Plugins/Developer/PerforceSourceControl/Source/PerforceSourceControl/Private/PerforceSourceControlOperations.cpp
@@ -2105,14 +2105,14 @@ bool FPerforceUpdateStatusWorker::UpdateStates() const
 			return;
 		}
 
-		const RowHandle RevisionControlTable = DataStorage->FindTable(FName("Editor_RevisionControlTable"));
-		if (RevisionControlTable == InvalidTableHandle)
+		const TableHandle PackageTable = DataStorage->FindTable(PackageTableName);
+		if (PackageTable == InvalidTableHandle)
 		{
 			return;
 		}
 		
 
-		auto GetRevisionControlRow = [DataStorage, RevisionControlTable](const FString& InFilename) -> RowHandle
+		auto GetPackageRow = [DataStorage, PackageTable](const FString& InFilename) -> RowHandle
 		{
 			FString Filename = InFilename;
 			FPaths::NormalizeFilename(Filename);
@@ -2123,13 +2123,13 @@ bool FPerforceUpdateStatusWorker::UpdateStates() const
 
 			if (!DataStorage->IsRowAvailable(Row))
 			{
-				Row = DataStorage->AddRow(RevisionControlTable);
+				Row = DataStorage->AddRow(PackageTable);
 				DataStorage->MapRow(MappingDomain, Key, Row);
 			}
 			return Row;
 		};
 
-		RowHandle Row = GetRevisionControlRow(State.GetFilename());
+		RowHandle Row = GetPackageRow(State.GetFilename());
 
 		if (!State.IsSourceControlled())
 		{
diff --git a/Engine/Plugins/Experimental/EditorDataStorageFeatures/Source/TedsRevisionControl/Private/Processors/RevisionControlProcessors.cpp b/Engine/Plugins/Experimental/EditorDataStorageFeatures/Source/TedsRevisionControl/Private/Processors/RevisionControlProcessors.cpp
--- a/Engine/Plugins/Experimental/EditorDataStorageFeatures/Source/TedsRevisionControl/Private/Processors/RevisionControlProcessors.cpp
+++ b/Engine/Plugins/Experimental/EditorDataStorageFeatures/Source/TedsRevisionControl/Private/Processors/RevisionControlProcessors.cpp
@@ -243,7 +243,7 @@ void URevisionControlDataStorageFactory::RegisterTables(UE::Editor::DataStorage:
 		TTypedElementColumnTypeList<
 			FTypedElementPackagePathColumn, FTypedElementPackageLoadedPathColumn,
 			FSCCRevisionIdColumn, FSCCExternalRevisionIdColumn>(),
-		FName("Editor_RevisionControlTable"));
+		PackageTableName);
 }
 
 void URevisionControlDataStorageFactory::RegisterQueries(UE::Editor::DataStorage::ICoreProvider& DataStorage)
diff --git a/Engine/Plugins/Experimental/EditorDataStorageFeatures/Source/TedsRevisionControl/Private/Queries/ObjectPackagePathToColumnQueries.cpp b/Engine/Plugins/Experimental/EditorDataStorageFeatures/Source/TedsRevisionControl/Private/Queries/ObjectPackagePathToColumnQueries.cpp
--- a/Engine/Plugins/Experimental/EditorDataStorageFeatures/Source/TedsRevisionControl/Private/Queries/ObjectPackagePathToColumnQueries.cpp
+++ b/Engine/Plugins/Experimental/EditorDataStorageFeatures/Source/TedsRevisionControl/Private/Queries/ObjectPackagePathToColumnQueries.cpp
@@ -160,9 +160,9 @@ void UTypedElementUObjectPackagePathFactory::RegisterTryAddPackageRef(UE::Editor
 				// Check in case this row was mapped before this command was pushed 
 				if (const RowHandle PackageRow = CoreProvider->LookupMappedRow(UE::Editor::RevisionControl::MappingDomain, FMapKey(PackagePathname)); !CoreProvider->IsRowAvailable(PackageRow))
 				{
-					if (const TableHandle RevisionControlTable = CoreProvider->FindTable(FName("Editor_RevisionControlTable")); RevisionControlTable != InvalidRowHandle)
+					if (const TableHandle PackageTable = CoreProvider->FindTable(PackageTableName); PackageTable != InvalidTableHandle)
 					{
-						const RowHandle Row = CoreProvider->AddRow(RevisionControlTable);
+						const RowHandle Row = CoreProvider->AddRow(PackageTable);
 						CoreProvider->MapRow(UE::Editor::RevisionControl::MappingDomain, FMapKey(PackagePathname), Row);
 					}
 				}
diff --git a/Engine/Source/Runtime/TypedElementFramework/Public/Elements/Columns/TypedElementPackageColumns.h b/Engine/Source/Runtime/TypedElementFramework/Public/Elements/Columns/TypedElementPackageColumns.h
--- a/Engine/Source/Runtime/TypedElementFramework/Public/Elements/Columns/TypedElementPackageColumns.h
+++ b/Engine/Source/Runtime/TypedElementFramework/Public/Elements/Columns/TypedElementPackageColumns.h
@@ -16,6 +16,7 @@
 
 namespace UE::Editor::DataStorage
 {
+	inline static const FName PackageTableName = TEXT("Editor_PackageTable");
 	inline static const FName PackageHierarchyName = TEXT("PackageHierarchy");
 }
 
