Perforce: Add optional functionality (that is by default disabled for now, but can be enabled via a 'p4.ParseIgnores' cvar) to mark files as Ignore instead of NotInDepot if they appear in a p4ignore rule.

#rnx

#Synced-CL 51621901

diff --git a/Engine/Plugins/Developer/PerforceSourceControl/Source/PerforceSourceControl/Private/PerforceConnection.cpp b/Engine/Plugins/Developer/PerforceSourceControl/Source/PerforceSourceControl/Private/PerforceConnection.cpp
--- a/Engine/Plugins/Developer/PerforceSourceControl/Source/PerforceSourceControl/Private/PerforceConnection.cpp
+++ b/Engine/Plugins/Developer/PerforceSourceControl/Source/PerforceSourceControl/Private/PerforceConnection.cpp
@@ -1136,6 +1136,19 @@ bool FPerforceConnection::IsValidConnection()
 	return bEstablishedConnection && !P4Client.Dropped();
 }
 
+bool FPerforceConnection::IsIgnored(const FString& InFile)
+{
+	if (Ignore* P4Ignore = P4Client.GetIgnore())
+	{
+		const StrPtr& IgnoreFile = P4Client.GetIgnoreFile();
+		const StrBuf File(FROM_TCHAR(*InFile, bIsUnicode));
+
+		return P4Ignore->Reject(File, IgnoreFile, nullptr) != 0;
+	}
+
+	return false;
+}
+
 void FPerforceConnection::Disconnect()
 {
 	Error P4Error;
diff --git a/Engine/Plugins/Developer/PerforceSourceControl/Source/PerforceSourceControl/Private/PerforceConnection.h b/Engine/Plugins/Developer/PerforceSourceControl/Source/PerforceSourceControl/Private/PerforceConnection.h
--- a/Engine/Plugins/Developer/PerforceSourceControl/Source/PerforceSourceControl/Private/PerforceConnection.h
+++ b/Engine/Plugins/Developer/PerforceSourceControl/Source/PerforceSourceControl/Private/PerforceConnection.h
@@ -117,6 +117,9 @@ public:
 	/** Returns true if connection is currently active */
 	bool IsValidConnection();
 
+	/** Returns true if file is ignored through p4ignore */
+	bool IsIgnored(const FString& InFile);
+
 	/** If the connection is valid, disconnect from the server */
 	void Disconnect();
 
diff --git a/Engine/Plugins/Developer/PerforceSourceControl/Source/PerforceSourceControl/Private/PerforceSourceControlOperations.cpp b/Engine/Plugins/Developer/PerforceSourceControl/Source/PerforceSourceControl/Private/PerforceSourceControlOperations.cpp
--- a/Engine/Plugins/Developer/PerforceSourceControl/Source/PerforceSourceControl/Private/PerforceSourceControlOperations.cpp
+++ b/Engine/Plugins/Developer/PerforceSourceControl/Source/PerforceSourceControl/Private/PerforceSourceControlOperations.cpp
@@ -45,6 +45,15 @@ static FAutoConsoleVariableRef CVarAlwaysBranchFilesOnCopy(
 	TEXT("Use legacy behavior of always branching a file in perforce when copying."),
 	ECVF_Default
 );
+
+bool GParseIgnores = false;
+static FAutoConsoleVariableRef CVarParseIgnores(
+	TEXT("p4.ParseIgnores"),
+	GParseIgnores,
+	TEXT("Use p4ignore behavior to exclude files from revision control consideration."),
+	ECVF_Default
+);
+
 }
 
 template<typename Type>
@@ -2020,6 +2029,22 @@ bool FPerforceUpdateStatusWorker::Execute(FPerforceSourceControlCommand& InComma
 			RemoveRedundantErrors(InCommand, TEXT("' is not under client's root '"));
 			RemoveRedundantErrors(InCommand, TEXT(" - file(s) not in client view"), false);
 			RemoveRedundantErrors(InCommand, TEXT(" - protected namespace - access denied"), false);
+
+			// Parse ignores for files that landed in NotInDepot as 'fstat' output does not differentiate between NotInDepot vs Ignore.
+			if (GParseIgnores)
+			{
+				for (FPerforceSourceControlState& State : OutStates)
+				{
+					if (State.State == EPerforceState::NotInDepot)
+					{
+						const bool bIsIgnored = Connection.IsIgnored(State.GetFilename());
+						if (bIsIgnored)
+						{
+							State.State = EPerforceState::Ignore;
+						}
+					}
+				}
+			}
 		}
 		else
 		{
diff --git a/Engine/Plugins/Developer/PerforceSourceControl/Source/PerforceSourceControl/Private/PerforceSourceControlPrivate.h b/Engine/Plugins/Developer/PerforceSourceControl/Source/PerforceSourceControl/Private/PerforceSourceControlPrivate.h
--- a/Engine/Plugins/Developer/PerforceSourceControl/Source/PerforceSourceControl/Private/PerforceSourceControlPrivate.h
+++ b/Engine/Plugins/Developer/PerforceSourceControl/Source/PerforceSourceControl/Private/PerforceSourceControlPrivate.h
@@ -7,4 +7,5 @@
 THIRD_PARTY_INCLUDES_START
 #include <p4/clientapi.h>
 #include <p4/i18napi.h>
+#include <p4/ignore.h>
 THIRD_PARTY_INCLUDES_END
