Fix a typo in the ISourceControlProvider api.

#rnx

- EInitFlags::SupressErrorLogging had a typo but is in the public api so we need to deprecate it and suggest people swittch to the new, correctly spelt version.
- EConnectionOptions::SupressErrorLogging can be updated without deprecating as it is in a private header.

#Synced-CL 51757651

diff --git a/Engine/Plugins/Developer/PerforceSourceControl/Source/PerforceSourceControl/Private/PerforceConnection.cpp b/Engine/Plugins/Developer/PerforceSourceControl/Source/PerforceSourceControl/Private/PerforceConnection.cpp
--- a/Engine/Plugins/Developer/PerforceSourceControl/Source/PerforceSourceControl/Private/PerforceConnection.cpp
+++ b/Engine/Plugins/Developer/PerforceSourceControl/Source/PerforceSourceControl/Private/PerforceConnection.cpp
@@ -1014,7 +1014,7 @@ bool FPerforceConnection::EnsureValidConnection(const FPerforceConnectionInfo& I
 	{
 		FinalizeErrors(SCCProvider, OutSettings, OutConnectionErrors);
 
-		if (!EnumHasAllFlags(Options, EConnectionOptions::SupressErrorLogging))
+		if (!EnumHasAllFlags(Options, EConnectionOptions::SuppressErrorLogging))
 		{
 			FTSMessageLog SourceControlLog("SourceControl");
 
diff --git a/Engine/Plugins/Developer/PerforceSourceControl/Source/PerforceSourceControl/Private/PerforceConnection.h b/Engine/Plugins/Developer/PerforceSourceControl/Source/PerforceSourceControl/Private/PerforceConnection.h
--- a/Engine/Plugins/Developer/PerforceSourceControl/Source/PerforceSourceControl/Private/PerforceConnection.h
+++ b/Engine/Plugins/Developer/PerforceSourceControl/Source/PerforceSourceControl/Private/PerforceConnection.h
@@ -48,7 +48,7 @@ enum class EConnectionOptions : uint8
 	/** The connection does not require a workspace to be considered valid */
 	WorkspaceOptional	= 1 << 0,
 	/** Errors will not be logged but will still be returned to the caller */
-	SupressErrorLogging	= 1 << 1
+	SuppressErrorLogging	= 1 << 1
 };
 ENUM_CLASS_FLAGS(EConnectionOptions);
 
diff --git a/Engine/Plugins/Developer/PerforceSourceControl/Source/PerforceSourceControl/Private/PerforceSourceControlProvider.cpp b/Engine/Plugins/Developer/PerforceSourceControl/Source/PerforceSourceControl/Private/PerforceSourceControlProvider.cpp
--- a/Engine/Plugins/Developer/PerforceSourceControl/Source/PerforceSourceControl/Private/PerforceSourceControlProvider.cpp
+++ b/Engine/Plugins/Developer/PerforceSourceControl/Source/PerforceSourceControl/Private/PerforceSourceControlProvider.cpp
@@ -315,9 +315,9 @@ ISourceControlProvider::FInitResult FPerforceSourceControlProvider::ParseCommand
 		FPerforceConnectionInfo ConnectionInfo = P4Settings.GetConnectionInfo();
 		FPerforceConnectionInfo OutputSettings;
 
-		if (EnumHasAnyFlags(InitFlags, EInitFlags::SupressErrorLogging))
+		if (EnumHasAnyFlags(InitFlags, EInitFlags::SuppressErrorLogging))
 		{
-			ConnectionOptions |= EConnectionOptions::SupressErrorLogging;
+			ConnectionOptions |= EConnectionOptions::SuppressErrorLogging;
 		}
 
 		if(FPerforceConnection::EnsureValidConnection(ConnectionInfo, *this, ConnectionOptions, OutputSettings, Result.Errors))
diff --git a/Engine/Source/Developer/SourceControl/Public/ISourceControlProvider.h b/Engine/Source/Developer/SourceControl/Public/ISourceControlProvider.h
--- a/Engine/Source/Developer/SourceControl/Public/ISourceControlProvider.h
+++ b/Engine/Source/Developer/SourceControl/Public/ISourceControlProvider.h
@@ -116,7 +116,8 @@ public:
 		/** A connection attempt to the provider will be made as part of the initialization */
 		AttemptConnection	= 1 << 0,
 		/** Errors will not be logged but will still be returned to the caller */
-		SupressErrorLogging = 1 << 1
+		SuppressErrorLogging = 1 << 1,
+		SupressErrorLogging UE_DEPRECATED(5.8, "Spelling typo, please use ISourceControlProvider::EInitFlags::SuppressErrorLogging instead") = 1 << 1
 	};
 
 	/** The result info returned form calls to  ISourceControlProvider::Init */
diff --git a/Engine/Source/Developer/Virtualization/Private/VirtualizationSourceControlBackend.cpp b/Engine/Source/Developer/Virtualization/Private/VirtualizationSourceControlBackend.cpp
--- a/Engine/Source/Developer/Virtualization/Private/VirtualizationSourceControlBackend.cpp
+++ b/Engine/Source/Developer/Virtualization/Private/VirtualizationSourceControlBackend.cpp
@@ -488,7 +488,7 @@ IVirtualizationBackend::EConnectionStatus FSourceControlBackend::OnConnectIntern
 	}
 
 	const ISourceControlProvider::EInitFlags InitOptions =	ISourceControlProvider::EInitFlags::AttemptConnection |
-															ISourceControlProvider::EInitFlags::SupressErrorLogging;
+															ISourceControlProvider::EInitFlags::SuppressErrorLogging;
 
 	ISourceControlProvider::FInitResult Result = SCCProvider->Init(InitOptions);
 
