Fix HLOD job revision control service not available
* Prevent idle perforce connections from auto closing when running from a commandlet.
* Commandlet are often long running processes, and the timeout value of 60m may not be sufficient.

diff --git a/Engine/Plugins/Developer/PerforceSourceControl/Source/PerforceSourceControl/Private/PerforceSourceControlProvider.cpp b/Engine/Plugins/Developer/PerforceSourceControl/Source/PerforceSourceControl/Private/PerforceSourceControlProvider.cpp
--- a/Engine/Plugins/Developer/PerforceSourceControl/Source/PerforceSourceControl/Private/PerforceSourceControlProvider.cpp
+++ b/Engine/Plugins/Developer/PerforceSourceControl/Source/PerforceSourceControl/Private/PerforceSourceControlProvider.cpp
@@ -754,7 +754,9 @@ void FPerforceSourceControlProvider::Tick()
 		OnSourceControlStateChanged.Broadcast();
 	}
 
-	if (PersistentConnection)
+	// Close idle connections when they haven't seen any activity for IdleConnectionDisconnectSeconds
+	// Disabled for commandlets, which are often long running processes
+	if (PersistentConnection && !IsRunningCommandlet())
 	{
 		const double Now = FPlatformTime::Seconds();
 
